Added missing loc undef

diff --git a/Engine/Plugins/Editor/SpeedTreeImporter/Source/SpeedTreeImporter/Private/ReimportSpeedTreeFactory.cpp b/Engine/Plugins/Editor/SpeedTreeImporter/Source/SpeedTreeImporter/Private/ReimportSpeedTreeFactory.cpp
--- a/Engine/Plugins/Editor/SpeedTreeImporter/Source/SpeedTreeImporter/Private/ReimportSpeedTreeFactory.cpp
+++ b/Engine/Plugins/Editor/SpeedTreeImporter/Source/SpeedTreeImporter/Private/ReimportSpeedTreeFactory.cpp
@@ -88,3 +88,5 @@ EReimportResult::Type UReimportSpeedTreeFactory::Reimport(UObject* Obj)
 }
 
 #endif
+
+#undef LOCTEXT_NAMESPACE
\ No newline at end of file
