Fix speedtree importer failing to compile when speedtree libs aren't present

diff --git a/Engine/Plugins/Editor/SpeedTreeImporter/Source/SpeedTreeImporter/Classes/ReimportSpeedTreeFactory.h b/Engine/Plugins/Editor/SpeedTreeImporter/Source/SpeedTreeImporter/Classes/ReimportSpeedTreeFactory.h
--- a/Engine/Plugins/Editor/SpeedTreeImporter/Source/SpeedTreeImporter/Classes/ReimportSpeedTreeFactory.h
+++ b/Engine/Plugins/Editor/SpeedTreeImporter/Source/SpeedTreeImporter/Classes/ReimportSpeedTreeFactory.h
@@ -14,11 +14,9 @@ class UReimportSpeedTreeFactory : public USpeedTreeImportFactory, public FReimpo
 	GENERATED_UCLASS_BODY()
 
 	// Begin FReimportHandler interface
-#if WITH_SPEEDTREE
 	virtual bool CanReimport(UObject* Obj, TArray<FString>& OutFilenames) OVERRIDE;
 	virtual void SetReimportPaths(UObject* Obj, const TArray<FString>& NewReimportPaths) OVERRIDE;
 	virtual EReimportResult::Type Reimport(UObject* Obj) OVERRIDE;
-#endif
 	// End FReimportHandler interface
 };
 
diff --git a/Engine/Plugins/Editor/SpeedTreeImporter/Source/SpeedTreeImporter/Private/ReimportSpeedTreeFactory.cpp b/Engine/Plugins/Editor/SpeedTreeImporter/Source/SpeedTreeImporter/Private/ReimportSpeedTreeFactory.cpp
--- a/Engine/Plugins/Editor/SpeedTreeImporter/Source/SpeedTreeImporter/Private/ReimportSpeedTreeFactory.cpp
+++ b/Engine/Plugins/Editor/SpeedTreeImporter/Source/SpeedTreeImporter/Private/ReimportSpeedTreeFactory.cpp
@@ -22,30 +22,33 @@ UReimportSpeedTreeFactory::UReimportSpeedTreeFactory(const class FPostConstructI
 	bText = false;
 }
 
-#if WITH_SPEEDTREE
-
 bool UReimportSpeedTreeFactory::CanReimport(UObject* Obj, TArray<FString>& OutFilenames)
 {
+#if WITH_SPEEDTREE
 	UStaticMesh* Mesh = Cast<UStaticMesh>(Obj);
 	if (Mesh && Mesh->AssetImportData)
 	{
 		OutFilenames.Add(FReimportManager::ResolveImportFilename(Mesh->AssetImportData->SourceFilePath, Mesh));
 		return true;
 	}
+#endif // #if WITH_SPEEDTREE
 	return false;
 }
 
 void UReimportSpeedTreeFactory::SetReimportPaths(UObject* Obj, const TArray<FString>& NewReimportPaths)
 {
+#if WITH_SPEEDTREE
 	UStaticMesh* Mesh = Cast<UStaticMesh>(Obj);
 	if (Mesh && ensure(NewReimportPaths.Num() == 1))
 	{
 		Mesh->AssetImportData->SourceFilePath = FReimportManager::ResolveImportFilename(NewReimportPaths[0], Mesh);
 	}
+#endif // #if WITH_SPEEDTREE
 }
 
 EReimportResult::Type UReimportSpeedTreeFactory::Reimport(UObject* Obj)
 {
+#if WITH_SPEEDTREE
 	UStaticMesh* Mesh = Cast<UStaticMesh>(Obj);
 	if (!Mesh)
 	{
@@ -83,10 +86,8 @@ EReimportResult::Type UReimportSpeedTreeFactory::Reimport(UObject* Obj)
 	{
 		UE_LOG(LogEditorFactories, Warning, TEXT("-- import failed"));
 	}
-
+#endif // #if WITH_SPEEDTREE
 	return EReimportResult::Failed;
 }
 
-#endif
-
-#undef LOCTEXT_NAMESPACE
\ No newline at end of file
+#undef LOCTEXT_NAMESPACE
