Case sensitiveness and pathname fixes from PR 185.

This is a subset of changes from pull request 185 by sbc100, 3dluvr, wshearn, abarbu et al: https://github.com/EpicGames/UnrealEngine/pull/185

diff --git a/Engine/Plugins/Editor/SpeedTreeImporter/Source/SpeedTreeImporter/Private/SpeedTreeImportFactory.cpp b/Engine/Plugins/Editor/SpeedTreeImporter/Source/SpeedTreeImporter/Private/SpeedTreeImportFactory.cpp
--- a/Engine/Plugins/Editor/SpeedTreeImporter/Source/SpeedTreeImporter/Private/SpeedTreeImportFactory.cpp
+++ b/Engine/Plugins/Editor/SpeedTreeImporter/Source/SpeedTreeImporter/Private/SpeedTreeImportFactory.cpp
@@ -18,7 +18,7 @@
 #include "Materials/MaterialExpressionTwoSidedSign.h"
 #include "Materials/MaterialExpressionVertexColor.h"
 
-#include "Mainframe.h"
+#include "MainFrame.h"
 #include "ModuleManager.h"
 #include "ObjectTools.h"
 #include "PackageTools.h"
diff --git a/Engine/Source/Editor/AnimGraph/Private/AnimationStateMachineGraph.cpp b/Engine/Source/Editor/AnimGraph/Private/AnimationStateMachineGraph.cpp
--- a/Engine/Source/Editor/AnimGraph/Private/AnimationStateMachineGraph.cpp
+++ b/Engine/Source/Editor/AnimGraph/Private/AnimationStateMachineGraph.cpp
@@ -1,7 +1,7 @@
 // Copyright 1998-2014 Epic Games, Inc. All Rights Reserved.
 
 #include "AnimGraphPrivatePCH.h"
-#include "AnimationStatemachineGraph.h"
+#include "AnimationStateMachineGraph.h"
 
 #define LOCTEXT_NAMESPACE "AnimationGraph"
 
diff --git a/Engine/Source/Editor/BlueprintGraph/Classes/K2Node_VariableGet.h b/Engine/Source/Editor/BlueprintGraph/Classes/K2Node_VariableGet.h
--- a/Engine/Source/Editor/BlueprintGraph/Classes/K2Node_VariableGet.h
+++ b/Engine/Source/Editor/BlueprintGraph/Classes/K2Node_VariableGet.h
@@ -2,7 +2,7 @@
 
 
 #pragma once
-#include "K2node_Variable.h"
+#include "K2Node_Variable.h"
 #include "K2Node_VariableGet.generated.h"
 
 UCLASS(MinimalAPI)
diff --git a/Engine/Source/Editor/GraphEditor/Private/AnimationNodes/SGraphNodeAnimationResult.cpp b/Engine/Source/Editor/GraphEditor/Private/AnimationNodes/SGraphNodeAnimationResult.cpp
--- a/Engine/Source/Editor/GraphEditor/Private/AnimationNodes/SGraphNodeAnimationResult.cpp
+++ b/Engine/Source/Editor/GraphEditor/Private/AnimationNodes/SGraphNodeAnimationResult.cpp
@@ -4,7 +4,7 @@
 #include "AnimGraphDefinitions.h"
 #include "SGraphNodeAnimationResult.h"
 #include "EdGraph/EdGraphNode.h"
-#include "AnimGraphNode_base.h"
+#include "AnimGraphNode_Base.h"
 
 /////////////////////////////////////////////////////
 // SGraphNodeAnimationResult
diff --git a/Engine/Source/Editor/Matinee/Private/MatineeViewportClient.cpp b/Engine/Source/Editor/Matinee/Private/MatineeViewportClient.cpp
--- a/Engine/Source/Editor/Matinee/Private/MatineeViewportClient.cpp
+++ b/Engine/Source/Editor/Matinee/Private/MatineeViewportClient.cpp
@@ -5,7 +5,7 @@
 
 #include "Runtime/Engine/Public/InterpolationHitProxy.h"
 #include "Runtime/Engine/Public/Slate/SceneViewport.h"
-#include "Matinee/MatineeACtor.h"
+#include "Matinee/MatineeActor.h"
 #include "Matinee/InterpTrackMove.h"
 #include "Matinee/InterpTrackEvent.h"
 
diff --git a/Engine/Source/Runtime/Launch/Private/Linux/LaunchLinux.cpp b/Engine/Source/Runtime/Launch/Private/Linux/LaunchLinux.cpp
--- a/Engine/Source/Runtime/Launch/Private/Linux/LaunchLinux.cpp
+++ b/Engine/Source/Runtime/Launch/Private/Linux/LaunchLinux.cpp
@@ -4,7 +4,7 @@
 #include "ExceptionHandling.h"
 
 #if WITH_EDITOR
-#include "Editor/MainFrame/Private/MainFrameActions.h"
+#include "Editor/MainFrame/Private/Frame/MainFrameActions.h"
 #endif
 
 #include <locale.h>
