Integrated Speedtree 7 Smooth LOD changes from IDV branch

diff --git a/Engine/Plugins/Editor/SpeedTreeImporter/Source/SpeedTreeImporter/Private/SpeedTreeImportFactory.cpp b/Engine/Plugins/Editor/SpeedTreeImporter/Source/SpeedTreeImporter/Private/SpeedTreeImportFactory.cpp
--- a/Engine/Plugins/Editor/SpeedTreeImporter/Source/SpeedTreeImporter/Private/SpeedTreeImportFactory.cpp
+++ b/Engine/Plugins/Editor/SpeedTreeImporter/Source/SpeedTreeImporter/Private/SpeedTreeImportFactory.cpp
@@ -12,11 +12,12 @@
 #include "Materials/MaterialExpressionPower.h"
 #include "Materials/MaterialExpressionMultiply.h"
 #include "Materials/MaterialExpressionSpeedTree.h"
-#include "Materials/MaterialExpressionSpeedTreeColorVariation.h"
 #include "Materials/MaterialExpressionTextureCoordinate.h"
 #include "Materials/MaterialExpressionTextureSample.h"
 #include "Materials/MaterialExpressionTwoSidedSign.h"
 #include "Materials/MaterialExpressionVertexColor.h"
+#include "Materials/MaterialExpressionMaterialFunctionCall.h"
+#include "Materials/MaterialFunction.h"
 
 #include "MainFrame.h"
 #include "ModuleManager.h"
@@ -422,7 +423,7 @@ static void LayoutMaterial(UMaterialInterface* MaterialInterface)
 	Material->EditorY = 0;
 
 	const int32 Height = 200;
-	const int32 Width = -250;
+	const int32 Width = 250;
 
 	// layout X to make sure each input is 1 step further than output connection
 	bool bContinue = true;
@@ -433,7 +434,7 @@ static void LayoutMaterial(UMaterialInterface* MaterialInterface)
 		for (int32 ExpressionIndex = 0; ExpressionIndex < Material->Expressions.Num(); ++ExpressionIndex)
 		{
 			UMaterialExpression* Expression = Material->Expressions[ExpressionIndex];
-			Expression->MaterialExpressionEditorX = FMath::Min(Expression->MaterialExpressionEditorX, Width);
+			Expression->MaterialExpressionEditorX = FMath::Min(Expression->MaterialExpressionEditorX, -Width);
 
 			TArray<FExpressionInput*> Inputs = Expression->GetInputs();
 			for (int32 InputIndex = 0; InputIndex < Inputs.Num(); ++InputIndex)
@@ -441,9 +442,9 @@ static void LayoutMaterial(UMaterialInterface* MaterialInterface)
 				UMaterialExpression* Input = Inputs[InputIndex]->Expression;
 				if (Input != NULL)
 				{
-					if (Input->MaterialExpressionEditorX > Expression->MaterialExpressionEditorX + Width)
+					if (Input->MaterialExpressionEditorX > Expression->MaterialExpressionEditorX - Width)
 					{
-						Input->MaterialExpressionEditorX = FMath::Min(Input->MaterialExpressionEditorX, Expression->MaterialExpressionEditorX + Width);
+						Input->MaterialExpressionEditorX = Expression->MaterialExpressionEditorX - Width;
 						bContinue = true;
 					}
 				}
@@ -461,7 +462,7 @@ static void LayoutMaterial(UMaterialInterface* MaterialInterface)
 		{
 			UMaterialExpression* Expression = Material->Expressions[ExpressionIndex];
 
-			if (Expression->MaterialExpressionEditorX == Width * Column)
+			if (Expression->MaterialExpressionEditorX == -Width * Column)
 			{
 				Expression->MaterialExpressionEditorY = 0;
 				int32 NumOutputs = 0;
@@ -519,7 +520,7 @@ static void LayoutMaterial(UMaterialInterface* MaterialInterface)
 
 		++Column;
 		bContinue = (ColumnExpressions.Num() > 0);
-	};
+	}
 }
 
 static UMaterialInterface* CreateSpeedTreeMaterial(UObject* Parent, FString MaterialFullName, const SpeedTree::SRenderState* RenderState, TSharedPtr<SSpeedTreeImportOptions> Options, ESpeedTreeWindType WindType, int NumBillboards)
@@ -551,7 +552,7 @@ static UMaterialInterface* CreateSpeedTreeMaterial(UObject* Parent, FString Mate
 		{
 			UnrealMaterial->BlendMode = BLEND_Masked;
 			UnrealMaterial->SetCastShadowAsMasked(true);
-			UnrealMaterial->TwoSided = !(RenderState->m_bFacingLeavesPresent || RenderState->m_bHorzBillboard || RenderState->m_bVertBillboard);
+			UnrealMaterial->TwoSided = !(RenderState->m_bHorzBillboard || RenderState->m_bVertBillboard);
 		}
 
 		// Notify the asset registry
@@ -779,11 +780,24 @@ static UMaterialInterface* CreateSpeedTreeMaterial(UObject* Parent, FString Mate
 	if (Options->IncludeColorAdjustment->IsChecked() && UnrealMaterial->BaseColor.Expression != NULL && 
 		(RenderState->m_bLeavesPresent || RenderState->m_bFacingLeavesPresent || RenderState->m_bVertBillboard || RenderState->m_bHorzBillboard))
 	{
-		UMaterialExpressionSpeedTreeColorVariation* ColorVariation = ConstructObject<UMaterialExpressionSpeedTreeColorVariation>(UMaterialExpressionSpeedTreeColorVariation::StaticClass(), UnrealMaterial);
-		UnrealMaterial->Expressions.Add(ColorVariation);
+		UMaterialFunction* ColorVariationFunction = LoadObject<UMaterialFunction>(NULL, TEXT("/Engine/Functions/Engine_MaterialFunctions01/SpeedTree/SpeedTreeColorVariation.SpeedTreeColorVariation"), NULL, LOAD_None, NULL);
+		if (ColorVariationFunction)
+		{
+			UMaterialExpressionMaterialFunctionCall* ColorVariation = ConstructObject<UMaterialExpressionMaterialFunctionCall>(UMaterialExpressionMaterialFunctionCall::StaticClass(), UnrealMaterial);
+			UnrealMaterial->Expressions.Add(ColorVariation);
+
+			ColorVariation->MaterialFunction = ColorVariationFunction;
+			ColorVariation->UpdateFromFunctionResource( );
 
-		ColorVariation->Input.Expression = UnrealMaterial->BaseColor.Expression;
-		UnrealMaterial->BaseColor.Expression = ColorVariation;
+			ColorVariation->GetInput(0)->Expression = UnrealMaterial->BaseColor.Expression;
+			ColorVariation->GetInput(0)->Mask = UnrealMaterial->BaseColor.Expression->GetOutputs()[0].Mask;
+			ColorVariation->GetInput(0)->MaskR = 1;
+			ColorVariation->GetInput(0)->MaskG = 1;
+			ColorVariation->GetInput(0)->MaskB = 1;
+			ColorVariation->GetInput(0)->MaskA = 0;
+
+			UnrealMaterial->BaseColor.Expression = ColorVariation;
+		}
 	}
 	
 	LayoutMaterial(UnrealMaterial);
@@ -944,7 +958,7 @@ UObject* USpeedTreeImportFactory::FactoryCreateBinary(UClass* InClass, UObject*
 #ifdef SPEEDTREE_KEY
 		SpeedTree::CCore::Authorize(PREPROCESSOR_TO_STRING(SPEEDTREE_KEY));
 #endif
-
+		
 		SpeedTree::CCore SpeedTree;
 		if (!SpeedTree.LoadTree(Buffer, BufferEnd - Buffer, false, false, Options->TreeScale))
 		{
@@ -961,7 +975,7 @@ UObject* USpeedTreeImportFactory::FactoryCreateBinary(UClass* InClass, UObject*
 			else
 			{
 				// make static mesh object
-				FString MeshName = ObjectTools::SanitizeObjectName(InName.ToString());
+				FString MeshName = ObjectTools::SanitizeObjectName(InName.ToString()) + TEXT("_Mesh");
 				FString NewPackageName = FPackageName::GetLongPackagePath(InParent->GetOutermost()->GetName()) + TEXT("/") + MeshName;
 				NewPackageName = PackageTools::SanitizePackageName(NewPackageName);
 				UPackage* Package = CreatePackage(NULL, *NewPackageName);
@@ -1314,7 +1328,7 @@ UObject* USpeedTreeImportFactory::FactoryCreateBinary(UClass* InClass, UObject*
 				if ((Options->ImportGeometryType == SSpeedTreeImportOptions::IGT_Billboards || Options->ImportGeometryType == SSpeedTreeImportOptions::IGT_Both) && 
 					SpeedTreeGeometry->m_sVertBBs.m_nNumBillboards > 0)
 				{
-					UMaterialInterface* Material = CreateSpeedTreeMaterial(InParent, MeshName + "_Billboard", &SpeedTreeGeometry->m_aBillboardRenderStates[SpeedTree::SHADER_PASS_MAIN], Options, WindType, SpeedTreeGeometry->m_sVertBBs.m_nNumBillboards);
+					UMaterialInterface* Material = CreateSpeedTreeMaterial(InParent, MeshName + "_Billboard", &SpeedTreeGeometry->m_aBillboardRenderStates[SpeedTree::RENDER_PASS_MAIN], Options, WindType, SpeedTreeGeometry->m_sVertBBs.m_nNumBillboards);
 					int32 MaterialIndex = StaticMesh->Materials.Num();
 					StaticMesh->Materials.Add(Material);
 
@@ -1356,7 +1370,7 @@ UObject* USpeedTreeImportFactory::FactoryCreateBinary(UClass* InClass, UObject*
 						FRotationMatrix BillboardRotate(Facing);
 
 						FVector TangentX = BillboardRotate.TransformVector(FVector(1.0f, 0.0f, 0.0f));
-						FVector TangentY = BillboardRotate.TransformVector(FVector(0.0f, 0.0f, 1.0f));
+						FVector TangentY = BillboardRotate.TransformVector(FVector(0.0f, 0.0f, -1.0f));
 						FVector TangentZ = BillboardRotate.TransformVector(FVector(0.0f, 1.0f, 0.0f));
 	
 						const float* TexCoords = &SpeedTreeGeometry->m_sVertBBs.m_pTexCoords[BillboardIndex * 4];
diff --git a/Engine/Plugins/Editor/SpeedTreeImporter/SpeedTreeImporter.uplugin b/Engine/Plugins/Editor/SpeedTreeImporter/SpeedTreeImporter.uplugin
--- a/Engine/Plugins/Editor/SpeedTreeImporter/SpeedTreeImporter.uplugin
+++ b/Engine/Plugins/Editor/SpeedTreeImporter/SpeedTreeImporter.uplugin
@@ -8,7 +8,7 @@
 	"CreatedByURL" : "http://www.speedtree.com",
 	"EngineVersion" : "4.2.0",
 	"Description" : "An importer for SpeedTree runtime (SRT) files.",
-	"Category" : "Editor",
+	"CategoryPath" : "Editor",
 	
 	"Modules" :
 	[
@@ -17,4 +17,4 @@
 			"Type" : "Editor"
 		}
 	]
-}
\ No newline at end of file
+}
diff --git a/Engine/Shaders/MaterialTemplate.usf b/Engine/Shaders/MaterialTemplate.usf
--- a/Engine/Shaders/MaterialTemplate.usf
+++ b/Engine/Shaders/MaterialTemplate.usf
@@ -660,12 +660,17 @@ float3 GetSpeedTreeVertexOffset(FMaterialVertexParameters Parameters, int Geomet
 
 	float3 TreePos = GetObjectWorldPosition(Parameters);
 
-	// geometry other than billboads can do smooth lod and wind effects (other than global) 
-	// extra data is required in the texcoords for this, which is only available on SpeedTrees
-	const float DistToTree = length(TreePos - View.ViewOrigin.xyz);
-	//const float LodInterp = saturate((DistToTree - SpeedTreeLODInfo.x) / SpeedTreeLODInfo.z);
-	//@todo - fix smooth LOD
-	const float LodInterp = 0.0f;
+	// compute LOD by finding screen space size
+	float LodInterp = 1.0;
+	if (LODType == SPEEDTREE_LOD_TYPE_SMOOTH) 
+	{
+		const float ScreenMultiple = 0.5 * max(View.ViewSizeAndSceneTexelSize.x * View.ViewToClip[0][0],
+												View.ViewSizeAndSceneTexelSize.y * View.ViewToClip[1][1]);
+		const float ScreenRadius = ScreenMultiple * Primitive.ObjectWorldPositionAndRadius.w / 
+									max(1.0, dot(TreePos - View.ViewOrigin, View.ViewForward));
+		const float ScreenArea = PI * ScreenRadius * ScreenRadius / (View.ViewSizeAndSceneTexelSize.x * View.ViewSizeAndSceneTexelSize.y);
+		LodInterp = saturate((ScreenArea - SpeedTreeLODInfo.x) / SpeedTreeLODInfo.z);
+	}
 
 	TreePos *= 0.001; // The only other use of the tree position is as an offset into trig functions, but big numbers don't play nice there
 
@@ -710,7 +715,7 @@ float3 GetSpeedTreeVertexOffset(FMaterialVertexParameters Parameters, int Geomet
 				{
 					float3 LODPos = float3(Parameters.TexCoords[3].x, Parameters.TexCoords[3].y, Parameters.TexCoords[4].x);
 					LODPos = mul(LODPos, LocalToWorld) / TreeScale;
-					FinalPosition = lerp(FinalPosition, LODPos, LodInterp);
+					FinalPosition = lerp(LODPos, FinalPosition, LodInterp);
 				}
 			#endif
 
@@ -741,12 +746,12 @@ float3 GetSpeedTreeVertexOffset(FMaterialVertexParameters Parameters, int Geomet
 					{
 						float3 LODPos = float3(Parameters.TexCoords[3].x, Parameters.TexCoords[3].y, Parameters.TexCoords[4].x);
 						LODPos = mul(LODPos, LocalToWorld) / TreeScale - Anchor;
-						FinalPosition = lerp(FinalPosition, LODPos, LodInterp);
+						FinalPosition = lerp(LODPos, FinalPosition, LodInterp);
 					}
 					else
 					{
 						float LODScalar = Parameters.TexCoords[3].x;
-						FinalPosition *= lerp(1.0, LODScalar, LodInterp);
+						FinalPosition *= lerp(LODScalar, 1.0, LodInterp);
 					}
 				}
 			#endif
@@ -792,43 +797,6 @@ float3 GetSpeedTreeVertexOffset(FMaterialVertexParameters Parameters, int Geomet
 
 #endif
 
-float3 GetSpeedTreeColorVariationBase(float3 InputColor, float Amount, bool bPreserveVibrance, float3 RandomColor, float Random) 
-{
-	Random -= 0.25;
-	float3 FinalColor = InputColor + RandomColor * Random * Amount;
-
-	if (bPreserveVibrance)
-	{
-		float OrigVibrance = max(InputColor.x, max(InputColor.y, InputColor.z));
-		float NewVibrance = max(FinalColor.x, max(FinalColor.y, FinalColor.z));
-
-		FinalColor *= OrigVibrance;
-		if (NewVibrance > 0)
-			FinalColor /= NewVibrance;
-	}
-
-	return saturate(FinalColor);
-}
-
-float3 GetSpeedTreeColorVariation(FMaterialVertexParameters Parameters, float3 InputColor, float Amount, bool bPreserveVibrance)
-{
-	#if USE_INSTANCING
-		return GetSpeedTreeColorVariationBase(InputColor, Amount, bPreserveVibrance, frac(float3(Parameters.PerInstanceParams.x * 100, Parameters.PerInstanceParams.x, Parameters.PerInstanceParams.x * 10)), Parameters.PerInstanceParams.x);
-	#else
-		return GetSpeedTreeColorVariationBase(InputColor, Amount, bPreserveVibrance, frac(Primitive.ObjectWorldPositionAndRadius.xyz * Primitive.ObjectWorldPositionAndRadius.zxy * 0.001), frac(dot(Primitive.ObjectWorldPositionAndRadius.xyz * 0.001, Primitive.LocalToWorld[0].xyz)));
-	#endif
-}
-
-float3 GetSpeedTreeColorVariation(FMaterialPixelParameters Parameters, float3 InputColor, float Amount, bool bPreserveVibrance)
-{
-	#if USE_INSTANCING
-		return GetSpeedTreeColorVariationBase(InputColor, Amount, bPreserveVibrance, frac(float3(Parameters.PerInstanceParams.x * 100, Parameters.PerInstanceParams.x, Parameters.PerInstanceParams.x * 10)), Parameters.PerInstanceParams.x);
-	#else
-		return GetSpeedTreeColorVariationBase(InputColor, Amount, bPreserveVibrance, frac(Primitive.ObjectWorldPositionAndRadius.xyz * Primitive.ObjectWorldPositionAndRadius.zxy * 0.001), frac(dot(Primitive.ObjectWorldPositionAndRadius.xyz * 0.001, Primitive.LocalToWorld[0].xyz)));
-	#endif
-}
-
-
 MaterialFloat2 GetLightmapUVs(FMaterialPixelParameters Parameters)
 {
 #if LIGHTMAP_UV_ACCESS
diff --git a/Engine/Source/Runtime/Engine/Classes/Materials/MaterialExpressionMaterialFunctionCall.h b/Engine/Source/Runtime/Engine/Classes/Materials/MaterialExpressionMaterialFunctionCall.h
--- a/Engine/Source/Runtime/Engine/Classes/Materials/MaterialExpressionMaterialFunctionCall.h
+++ b/Engine/Source/Runtime/Engine/Classes/Materials/MaterialExpressionMaterialFunctionCall.h
@@ -127,7 +127,7 @@ class UMaterialExpressionMaterialFunctionCall : public UMaterialExpression
 	 * Update FunctionInputs and FunctionOutputs from the MaterialFunction.  
 	 * This must be called to keep the inputs and outputs up to date with the function being used. 
 	 */
-	void UpdateFromFunctionResource();
+	ENGINE_API void UpdateFromFunctionResource();
 
 private:
 	
diff --git a/Engine/Source/Runtime/Engine/Classes/Materials/MaterialExpressionSpeedTreeColorVariation.h b/Engine/Source/Runtime/Engine/Classes/Materials/MaterialExpressionSpeedTreeColorVariation.h
deleted file mode 100644
diff --git a/Engine/Source/Runtime/Engine/Private/LocalVertexFactory.cpp b/Engine/Source/Runtime/Engine/Private/LocalVertexFactory.cpp
--- a/Engine/Source/Runtime/Engine/Private/LocalVertexFactory.cpp
+++ b/Engine/Source/Runtime/Engine/Private/LocalVertexFactory.cpp
@@ -33,8 +33,7 @@ public:
 
 				if (LODParameter.IsBound())
 				{
-					//@todo - speedtree smooth LOD
-					FVector LODData(0, 0, 0);
+					FVector LODData(BatchElement.MinScreenSize, BatchElement.MaxScreenSize, BatchElement.MaxScreenSize - BatchElement.MinScreenSize);
 					SetShaderValue(RHICmdList, Shader->GetVertexShader(), LODParameter, LODData);
 				}
 			}
diff --git a/Engine/Source/Runtime/Engine/Private/Materials/HLSLMaterialTranslator.h b/Engine/Source/Runtime/Engine/Private/Materials/HLSLMaterialTranslator.h
--- a/Engine/Source/Runtime/Engine/Private/Materials/HLSLMaterialTranslator.h
+++ b/Engine/Source/Runtime/Engine/Private/Materials/HLSLMaterialTranslator.h
@@ -3470,23 +3470,6 @@ protected:
 		}
 	}
 
-	/**
-	* Handles SpeedTree color variation between instances
-	*
-	* @return	Code index
-	*/
-	virtual int32 SpeedTreeColorVariation(int32 InputColor, float Amount, bool bPreserveVibrance) OVERRIDE 
-	{ 
-		if (ShaderFrequency != SF_Pixel && ShaderFrequency != SF_Vertex)
-		{
-			return NonVertexOrPixelShaderExpressionError();
-		}
-		else
-		{
-			return AddCodeChunk(MCT_Float3, TEXT("GetSpeedTreeColorVariation(Parameters, %s, %g, %s)"), *GetParameterCode(InputColor), Amount, bPreserveVibrance ? TEXT("true") : TEXT("false"));
-		}
-	}
-
 	/**
 	 * Adds code for texture coordinate offset to localize large UV
 	 *
diff --git a/Engine/Source/Runtime/Engine/Private/Materials/MaterialExpressions.cpp b/Engine/Source/Runtime/Engine/Private/Materials/MaterialExpressions.cpp
--- a/Engine/Source/Runtime/Engine/Private/Materials/MaterialExpressions.cpp
+++ b/Engine/Source/Runtime/Engine/Private/Materials/MaterialExpressions.cpp
@@ -106,7 +106,6 @@
 #include "Materials/MaterialExpressionScreenPosition.h"
 #include "Materials/MaterialExpressionSine.h"
 #include "Materials/MaterialExpressionSpeedTree.h"
-#include "Materials/MaterialExpressionSpeedTreeColorVariation.h"
 #include "Materials/MaterialExpressionSphereMask.h"
 #include "Materials/MaterialExpressionSphericalParticleOpacity.h"
 #include "Materials/MaterialExpressionSquareRoot.h"
@@ -9441,39 +9440,6 @@ bool UMaterialExpressionSpeedTree::CanEditChange(const UProperty* InProperty) co
 
 #endif
 
-/*------------------------------------------------------------------------------
-SpeedTree Color Variation material expression.
-------------------------------------------------------------------------------*/
-UMaterialExpressionSpeedTreeColorVariation::UMaterialExpressionSpeedTreeColorVariation(const class FPostConstructInitializeProperties& PCIP)
-: Super(PCIP)
-{
-	// Structure to hold one-time initialization
-	struct FConstructorStatics
-	{
-		FString NAME_SpeedTree;
-		FConstructorStatics()
-			: NAME_SpeedTree(LOCTEXT("SpeedTree", "SpeedTree").ToString())
-		{
-		}
-	};
-	static FConstructorStatics ConstructorStatics;
-
-	Amount = 0.1f;
-	bPreserveVibrance = true;
-
-	MenuCategories.Add(ConstructorStatics.NAME_SpeedTree);
-}
-
-int32 UMaterialExpressionSpeedTreeColorVariation::Compile(class FMaterialCompiler* Compiler, int32 OutputIndex, int32 MultiplexIndex)
-{
-	return Compiler->SpeedTreeColorVariation(Input.Compile(Compiler), Amount, bPreserveVibrance);
-}
-
-void UMaterialExpressionSpeedTreeColorVariation::GetCaption(TArray<FString>& OutCaptions) const
-{
-	OutCaptions.Add(TEXT("SpeedTree Color Variation"));
-}
-
 
 ///////////////////////////////////////////////////////////////////////////////
 // UMaterialExpressionEyeAdaptation
diff --git a/Engine/Source/Runtime/Engine/Private/StaticMesh.cpp b/Engine/Source/Runtime/Engine/Private/StaticMesh.cpp
--- a/Engine/Source/Runtime/Engine/Private/StaticMesh.cpp
+++ b/Engine/Source/Runtime/Engine/Private/StaticMesh.cpp
@@ -2532,6 +2532,13 @@ bool UStaticMesh::CanLODsShareStaticLighting() const
 	{
 		bCanShareData = bCanShareData && SourceModels[LODIndex].RawMeshBulkData->IsEmpty();
 	}
+
+	if (SpeedTreeWind.IsValid())
+	{
+		// SpeedTrees are set up for lighting to share between LODs
+		bCanShareData = true;
+	}
+
 	return bCanShareData;
 }
 
diff --git a/Engine/Source/Runtime/Engine/Private/StaticMeshRender.cpp b/Engine/Source/Runtime/Engine/Private/StaticMeshRender.cpp
--- a/Engine/Source/Runtime/Engine/Private/StaticMeshRender.cpp
+++ b/Engine/Source/Runtime/Engine/Private/StaticMeshRender.cpp
@@ -167,6 +167,12 @@ bool FStaticMeshSceneProxy::GetShadowMeshElement(int32 LODIndex, uint8 InDepthPr
 	OutMeshElement.CastShadow = true;
 	OutMeshElement.LODIndex = LODIndex;
 	OutMeshElement.LCI = &ProxyLODInfo;
+	OutBatchElement.MinScreenSize = 0.0f;
+	if (LODIndex < MAX_STATIC_MESH_LODS - 1)
+	{
+		OutBatchElement.MinScreenSize = GetScreenSize(LODIndex + 1);
+	}
+	OutBatchElement.MaxScreenSize = GetScreenSize(LODIndex);
 
 	return true;
 }
@@ -215,6 +221,12 @@ bool FStaticMeshSceneProxy::GetMeshElement(int32 LODIndex,int32 SectionIndex,uin
 		OutMeshElement.ReverseCulling = IsLocalToWorldDeterminantNegative();
 		OutMeshElement.CastShadow = bCastShadow && Section.bCastShadow;
 		OutMeshElement.DepthPriorityGroup = (ESceneDepthPriorityGroup)InDepthPriorityGroup;
+		OutBatchElement.MinScreenSize = 0.0f;
+		if (LODIndex < MAX_STATIC_MESH_LODS - 1)
+		{
+			OutBatchElement.MinScreenSize = GetScreenSize(LODIndex + 1);
+		}
+		OutBatchElement.MaxScreenSize = GetScreenSize(LODIndex);
 		return true;
 	}
 	else
@@ -238,6 +250,12 @@ bool FStaticMeshSceneProxy::GetWireframeMeshElement(int32 LODIndex, const FMater
 	OutMeshElement.ReverseCulling = IsLocalToWorldDeterminantNegative();
 	OutMeshElement.CastShadow = bCastShadow;
 	OutMeshElement.DepthPriorityGroup = (ESceneDepthPriorityGroup)InDepthPriorityGroup;
+	OutBatchElement.MinScreenSize = 0.0f;
+	if (LODIndex < MAX_STATIC_MESH_LODS - 1)
+	{
+		OutBatchElement.MinScreenSize = GetScreenSize(LODIndex + 1);
+	}
+	OutBatchElement.MaxScreenSize = GetScreenSize(LODIndex);
 
 	const bool bWireframe = true;
 	const bool bRequiresAdjacencyInformation = false;
diff --git a/Engine/Source/Runtime/Engine/Public/MaterialCompiler.h b/Engine/Source/Runtime/Engine/Public/MaterialCompiler.h
--- a/Engine/Source/Runtime/Engine/Public/MaterialCompiler.h
+++ b/Engine/Source/Runtime/Engine/Public/MaterialCompiler.h
@@ -169,7 +169,6 @@ public:
 	virtual int32 DepthOfFieldFunction(int32 Depth, int32 FunctionValueIndex) = 0;
 	virtual int32 AtmosphericFogColor(int32 WorldPosition) = 0;
 	virtual int32 SpeedTree(ESpeedTreeGeometryType GeometryType, ESpeedTreeWindType WindType, ESpeedTreeLODType LODType, float BillboardThreshold) = 0;
-	virtual int32 SpeedTreeColorVariation(int32 InputColor, float Amount, bool bPreserveVibrance) = 0;
 	virtual int32 TextureCoordinateOffset() = 0;
 	virtual int32 EyeAdaptation() = 0;
 };
@@ -332,11 +331,6 @@ public:
 	{ 
 		return Compiler->SpeedTree(GeometryType, WindType, LODType, BillboardThreshold); 
 	}
-	
-	virtual int32 SpeedTreeColorVariation(int32 InputColor, float Amount, bool bPreserveVibrance) OVERRIDE
-	{
-		return Compiler->SpeedTreeColorVariation(InputColor, Amount, bPreserveVibrance);
-	}
 
 	virtual int32 AtmosphericFogColor(int32 WorldPosition) OVERRIDE
 	{
diff --git a/Engine/Source/Runtime/Engine/Public/SceneManagement.h b/Engine/Source/Runtime/Engine/Public/SceneManagement.h
--- a/Engine/Source/Runtime/Engine/Public/SceneManagement.h
+++ b/Engine/Source/Runtime/Engine/Public/SceneManagement.h
@@ -1251,6 +1251,8 @@ struct FMeshBatchElement
 	uint32 MaxVertexIndex;
 	int32 UserIndex;
 	void* UserData;
+	float MinScreenSize;
+	float MaxScreenSize;
 
 	/** 
 	 *	DynamicIndexData - pointer to user memory containing the index data.
@@ -1265,6 +1267,8 @@ struct FMeshBatchElement
 	,	NumInstances(1)
 	,	UserIndex(-1)
 	,	UserData(NULL)
+	,	MinScreenSize(0.0f)
+	,	MaxScreenSize(1.0f)
 	,	DynamicIndexData(NULL)
 	{
 	}
