Merging

//depot/Partners/IDV/UE4-ST7(2)/Engine/...

to //depot/UE4/Engine/...

diff --git a/Engine/Plugins/Editor/SpeedTreeImporter/Source/SpeedTreeImporter/Private/SpeedTreeImportFactory.cpp b/Engine/Plugins/Editor/SpeedTreeImporter/Source/SpeedTreeImporter/Private/SpeedTreeImportFactory.cpp
--- a/Engine/Plugins/Editor/SpeedTreeImporter/Source/SpeedTreeImporter/Private/SpeedTreeImportFactory.cpp
+++ b/Engine/Plugins/Editor/SpeedTreeImporter/Source/SpeedTreeImporter/Private/SpeedTreeImportFactory.cpp
@@ -368,7 +368,7 @@ static UTexture* CreateSpeedTreeMaterialTexture(UObject* Parent, FString Filenam
 	}
 
 	FString Extension = FPaths::GetExtension(Filename).ToLower();
-	FString TextureName = FPaths::GetBaseFilename(Filename);
+	FString TextureName = FPaths::GetBaseFilename(Filename) + TEXT("_Tex");
 	TextureName = ObjectTools::SanitizeObjectName(TextureName);
 
 	// set where to place the textures
@@ -532,12 +532,13 @@ static UMaterialInterface* CreateSpeedTreeMaterial(UObject* Parent, FString Mate
 	}
 	
 	// set where to place the materials
-	FString NewPackageName = FPackageName::GetLongPackagePath(Parent->GetOutermost()->GetName()) + TEXT("/") + MaterialFullName;
+	FString FixedMaterialName = MaterialFullName + TEXT("_Mat");
+	FString NewPackageName = FPackageName::GetLongPackagePath(Parent->GetOutermost()->GetName()) + TEXT("/") + FixedMaterialName;
 	NewPackageName = PackageTools::SanitizePackageName(NewPackageName);
 	UPackage* Package = CreatePackage(NULL, *NewPackageName);
 	
 	// does not override existing materials
-	UMaterialInterface* UnrealMaterialInterface = FindObject<UMaterialInterface>(Package, *MaterialFullName);
+	UMaterialInterface* UnrealMaterialInterface = FindObject<UMaterialInterface>(Package, *FixedMaterialName);
 	if (UnrealMaterialInterface != NULL)
 	{
 		return UnrealMaterialInterface;
@@ -545,7 +546,7 @@ static UMaterialInterface* CreateSpeedTreeMaterial(UObject* Parent, FString Mate
 	
 	// create an unreal material asset
 	UMaterialFactoryNew* MaterialFactory = new UMaterialFactoryNew(FPostConstructInitializeProperties());
-	UMaterial* UnrealMaterial = (UMaterial*)MaterialFactory->FactoryCreateNew(UMaterial::StaticClass(), Package, *MaterialFullName, RF_Standalone|RF_Public, NULL, GWarn);
+	UMaterial* UnrealMaterial = (UMaterial*)MaterialFactory->FactoryCreateNew(UMaterial::StaticClass(), Package, *FixedMaterialName, RF_Standalone|RF_Public, NULL, GWarn);
 	if (UnrealMaterial != NULL)
 	{
 		if (!RenderState->m_bDiffuseAlphaMaskIsOpaque && !RenderState->m_bBranchesPresent && !RenderState->m_bRigidMeshesPresent)
@@ -975,7 +976,7 @@ UObject* USpeedTreeImportFactory::FactoryCreateBinary(UClass* InClass, UObject*
 			else
 			{
 				// make static mesh object
-				FString MeshName = ObjectTools::SanitizeObjectName(InName.ToString()) + TEXT("_Mesh");
+				FString MeshName = ObjectTools::SanitizeObjectName(InName.ToString());
 				FString NewPackageName = FPackageName::GetLongPackagePath(InParent->GetOutermost()->GetName()) + TEXT("/") + MeshName;
 				NewPackageName = PackageTools::SanitizePackageName(NewPackageName);
 				UPackage* Package = CreatePackage(NULL, *NewPackageName);
@@ -1035,6 +1036,24 @@ UObject* USpeedTreeImportFactory::FactoryCreateBinary(UClass* InClass, UObject*
 					}
 				}
 
+				// Force LOD code out of the shaders if we only have one LOD
+				if (Options->IncludeSmoothLODCheck->IsChecked( ))
+				{
+					int32 TotalLODs = 0;
+					if (Options->ImportGeometryType != SSpeedTreeImportOptions::IGT_Billboards)
+					{
+						TotalLODs += SpeedTreeGeometry->m_nNumLods;
+					}
+					if (Options->ImportGeometryType != SSpeedTreeImportOptions::IGT_3D && SpeedTreeGeometry->m_sVertBBs.m_nNumBillboards > 0)
+					{
+						++TotalLODs;
+					}
+					if (TotalLODs < 2)
+					{
+						Options->IncludeSmoothLODCheck->ToggleCheckedState( );
+					}
+				}
+
 				// make geometry LODs
 				if (Options->ImportGeometryType != SSpeedTreeImportOptions::IGT_Billboards)
 				{
@@ -1312,7 +1331,8 @@ UObject* USpeedTreeImportFactory::FactoryCreateBinary(UClass* InClass, UObject*
 						LODModel->BuildSettings.bRecomputeNormals = false;
 						LODModel->BuildSettings.bRecomputeTangents = false;
 						LODModel->BuildSettings.bRemoveDegenerates = true;
-						LODModel->BuildSettings.bUseFullPrecisionUVs = false;
+						LODModel->BuildSettings.bUseFullPrecisionUVs = false;	
+						LODModel->ScreenSize = 0.1f / FMath::Pow(2.0f, StaticMesh->SourceModels.Num() - 1);
 						LODModel->RawMeshBulkData->SaveRawMesh(RawMesh);
 
 						for (int32 MaterialIndex = 0; MaterialIndex < StaticMesh->Materials.Num(); ++MaterialIndex)
@@ -1325,8 +1345,7 @@ UObject* USpeedTreeImportFactory::FactoryCreateBinary(UClass* InClass, UObject*
 				}
 
 				// make billboard LOD
-				if ((Options->ImportGeometryType == SSpeedTreeImportOptions::IGT_Billboards || Options->ImportGeometryType == SSpeedTreeImportOptions::IGT_Both) && 
-					SpeedTreeGeometry->m_sVertBBs.m_nNumBillboards > 0)
+				if (Options->ImportGeometryType != SSpeedTreeImportOptions::IGT_3D && SpeedTreeGeometry->m_sVertBBs.m_nNumBillboards > 0)
 				{
 					UMaterialInterface* Material = CreateSpeedTreeMaterial(InParent, MeshName + "_Billboard", &SpeedTreeGeometry->m_aBillboardRenderStates[SpeedTree::RENDER_PASS_MAIN], Options, WindType, SpeedTreeGeometry->m_sVertBBs.m_nNumBillboards);
 					int32 MaterialIndex = StaticMesh->Materials.Num();
@@ -1424,6 +1443,7 @@ UObject* USpeedTreeImportFactory::FactoryCreateBinary(UClass* InClass, UObject*
 					LODModel->BuildSettings.bRecomputeTangents = false;
 					LODModel->BuildSettings.bRemoveDegenerates = true;
 					LODModel->BuildSettings.bUseFullPrecisionUVs = false;
+					LODModel->ScreenSize = 0.1f / FMath::Pow(2.0f, StaticMesh->SourceModels.Num() - 1);
 					LODModel->RawMeshBulkData->SaveRawMesh(RawMesh);
 				}
 
@@ -1438,6 +1458,10 @@ UObject* USpeedTreeImportFactory::FactoryCreateBinary(UClass* InClass, UObject*
 						MakeBodyFromCollisionObjects(StaticMesh, CollisionObjects, NumCollisionObjects);
 					}
 				}
+
+				// make better LOD info for SpeedTrees
+				StaticMesh->bAutoComputeLODScreenSize = false;
+				StaticMesh->bRequiresLODDistanceConversion = false;
 			}
 		}
 	}
diff --git a/Engine/Source/Runtime/Engine/Private/StaticMeshRender.cpp b/Engine/Source/Runtime/Engine/Private/StaticMeshRender.cpp
--- a/Engine/Source/Runtime/Engine/Private/StaticMeshRender.cpp
+++ b/Engine/Source/Runtime/Engine/Private/StaticMeshRender.cpp
@@ -167,12 +167,20 @@ bool FStaticMeshSceneProxy::GetShadowMeshElement(int32 LODIndex, uint8 InDepthPr
 	OutMeshElement.CastShadow = true;
 	OutMeshElement.LODIndex = LODIndex;
 	OutMeshElement.LCI = &ProxyLODInfo;
-	OutBatchElement.MinScreenSize = 0.0f;
-	if (LODIndex < MAX_STATIC_MESH_LODS - 1)
+	if (ForcedLodModel > 0) 
 	{
-		OutBatchElement.MinScreenSize = GetScreenSize(LODIndex + 1);
+		OutBatchElement.MaxScreenSize = 0.0f;
+		OutBatchElement.MinScreenSize = -1.0f;
+	}
+	else
+	{
+		OutBatchElement.MaxScreenSize = GetScreenSize(LODIndex);
+		OutBatchElement.MinScreenSize = 0.0f;
+		if (LODIndex < MAX_STATIC_MESH_LODS - 1)
+		{
+			OutBatchElement.MinScreenSize = GetScreenSize(LODIndex + 1);
+		}
 	}
-	OutBatchElement.MaxScreenSize = GetScreenSize(LODIndex);
 
 	return true;
 }
@@ -221,12 +229,21 @@ bool FStaticMeshSceneProxy::GetMeshElement(int32 LODIndex,int32 SectionIndex,uin
 		OutMeshElement.ReverseCulling = IsLocalToWorldDeterminantNegative();
 		OutMeshElement.CastShadow = bCastShadow && Section.bCastShadow;
 		OutMeshElement.DepthPriorityGroup = (ESceneDepthPriorityGroup)InDepthPriorityGroup;
-		OutBatchElement.MinScreenSize = 0.0f;
-		if (LODIndex < MAX_STATIC_MESH_LODS - 1)
+		if (ForcedLodModel > 0) 
 		{
-			OutBatchElement.MinScreenSize = GetScreenSize(LODIndex + 1);
+			OutBatchElement.MaxScreenSize = 0.0f;
+			OutBatchElement.MinScreenSize = -1.0f;
 		}
-		OutBatchElement.MaxScreenSize = GetScreenSize(LODIndex);
+		else
+		{
+			OutBatchElement.MaxScreenSize = GetScreenSize(LODIndex);
+			OutBatchElement.MinScreenSize = 0.0f;
+			if (LODIndex < MAX_STATIC_MESH_LODS - 1)
+			{
+				OutBatchElement.MinScreenSize = GetScreenSize(LODIndex + 1);
+			}
+		}
+			
 		return true;
 	}
 	else
@@ -250,12 +267,20 @@ bool FStaticMeshSceneProxy::GetWireframeMeshElement(int32 LODIndex, const FMater
 	OutMeshElement.ReverseCulling = IsLocalToWorldDeterminantNegative();
 	OutMeshElement.CastShadow = bCastShadow;
 	OutMeshElement.DepthPriorityGroup = (ESceneDepthPriorityGroup)InDepthPriorityGroup;
-	OutBatchElement.MinScreenSize = 0.0f;
-	if (LODIndex < MAX_STATIC_MESH_LODS - 1)
+	if (ForcedLodModel > 0) 
+	{
+		OutBatchElement.MaxScreenSize = 0.0f;
+		OutBatchElement.MinScreenSize = -1.0f;
+	}
+	else
 	{
-		OutBatchElement.MinScreenSize = GetScreenSize(LODIndex + 1);
+		OutBatchElement.MaxScreenSize = GetScreenSize(LODIndex);
+		OutBatchElement.MinScreenSize = 0.0f;
+		if (LODIndex < MAX_STATIC_MESH_LODS - 1)
+		{
+			OutBatchElement.MinScreenSize = GetScreenSize(LODIndex + 1);
+		}
 	}
-	OutBatchElement.MaxScreenSize = GetScreenSize(LODIndex);
 
 	const bool bWireframe = true;
 	const bool bRequiresAdjacencyInformation = false;
