Merging

//depot/Partners/IDV/UE4-ST7(2)/Engine/Plugins/Editor/SpeedTreeImporter/...

to //depot/UE4/Engine/Plugins/Editor/SpeedTreeImporter/...

diff --git a/Engine/Plugins/Editor/SpeedTreeImporter/Source/SpeedTreeImporter/Private/SpeedTreeImportFactory.cpp b/Engine/Plugins/Editor/SpeedTreeImporter/Source/SpeedTreeImporter/Private/SpeedTreeImportFactory.cpp
--- a/Engine/Plugins/Editor/SpeedTreeImporter/Source/SpeedTreeImporter/Private/SpeedTreeImportFactory.cpp
+++ b/Engine/Plugins/Editor/SpeedTreeImporter/Source/SpeedTreeImporter/Private/SpeedTreeImportFactory.cpp
@@ -897,8 +897,8 @@ static void MakeBodyFromCollisionObjects(UStaticMesh* StaticMesh, const SpeedTre
 	for (int32 CollisionObjectIndex = 0; CollisionObjectIndex < NumCollisionObjects; ++CollisionObjectIndex)
 	{
 		const SpeedTree::SCollisionObject& CollisionObject = CollisionObjects[CollisionObjectIndex];
-		const FVector Pos1(CollisionObject.m_vCenter1.x, CollisionObject.m_vCenter1.y, CollisionObject.m_vCenter1.z);
-		const FVector Pos2(CollisionObject.m_vCenter2.x, CollisionObject.m_vCenter2.y, CollisionObject.m_vCenter2.z);
+		const FVector Pos1(-CollisionObject.m_vCenter1.x, CollisionObject.m_vCenter1.y, CollisionObject.m_vCenter1.z);
+		const FVector Pos2(-CollisionObject.m_vCenter2.x, CollisionObject.m_vCenter2.y, CollisionObject.m_vCenter2.z);
 
 		if (Pos1 == Pos2)
 		{
