Rename FPostConstructInitializeProperties to something simpler

Changed it with FObjectInitializer.

UECORE-7

diff --git a/Engine/Plugins/2D/Paper2D/Source/Paper2D/Private/PaperCharacter.cpp b/Engine/Plugins/2D/Paper2D/Source/Paper2D/Private/PaperCharacter.cpp
--- a/Engine/Plugins/2D/Paper2D/Source/Paper2D/Private/PaperCharacter.cpp
+++ b/Engine/Plugins/2D/Paper2D/Source/Paper2D/Private/PaperCharacter.cpp
@@ -7,11 +7,11 @@ FName APaperCharacter::SpriteComponentName(TEXT("Sprite0"));
 //////////////////////////////////////////////////////////////////////////
 // APaperCharacter
 
-APaperCharacter::APaperCharacter(const class FPostConstructInitializeProperties& PCIP)
-	: Super(PCIP.DoNotCreateDefaultSubobject(ACharacter::MeshComponentName))
+APaperCharacter::APaperCharacter(const FObjectInitializer& ObjectInitializer)
+	: Super(ObjectInitializer.DoNotCreateDefaultSubobject(ACharacter::MeshComponentName))
 {
 	// Try to create the sprite component
-	Sprite = PCIP.CreateOptionalDefaultSubobject<UPaperFlipbookComponent>(this, APaperCharacter::SpriteComponentName);
+	Sprite = ObjectInitializer.CreateOptionalDefaultSubobject<UPaperFlipbookComponent>(this, APaperCharacter::SpriteComponentName);
 	if (Sprite)
 	{
 		Sprite->AlwaysLoadOnClient = true;
diff --git a/Engine/Plugins/2D/Paper2D/Source/Paper2D/Private/PaperFlipbook.cpp b/Engine/Plugins/2D/Paper2D/Source/Paper2D/Private/PaperFlipbook.cpp
--- a/Engine/Plugins/2D/Paper2D/Source/Paper2D/Private/PaperFlipbook.cpp
+++ b/Engine/Plugins/2D/Paper2D/Source/Paper2D/Private/PaperFlipbook.cpp
@@ -11,8 +11,8 @@
 //////////////////////////////////////////////////////////////////////////
 // UPaperFlipbook
 
-UPaperFlipbook::UPaperFlipbook(const FPostConstructInitializeProperties& PCIP)
-	: Super(PCIP)
+UPaperFlipbook::UPaperFlipbook(const FObjectInitializer& ObjectInitializer)
+	: Super(ObjectInitializer)
 {
 	FramesPerSecond = 15;
 
diff --git a/Engine/Plugins/2D/Paper2D/Source/Paper2D/Private/PaperFlipbookActor.cpp b/Engine/Plugins/2D/Paper2D/Source/Paper2D/Private/PaperFlipbookActor.cpp
--- a/Engine/Plugins/2D/Paper2D/Source/Paper2D/Private/PaperFlipbookActor.cpp
+++ b/Engine/Plugins/2D/Paper2D/Source/Paper2D/Private/PaperFlipbookActor.cpp
@@ -6,10 +6,10 @@
 //////////////////////////////////////////////////////////////////////////
 // APaperFlipbookActor
 
-APaperFlipbookActor::APaperFlipbookActor(const FPostConstructInitializeProperties& PCIP)
-	: Super(PCIP)
+APaperFlipbookActor::APaperFlipbookActor(const FObjectInitializer& ObjectInitializer)
+	: Super(ObjectInitializer)
 {
-	RenderComponent = PCIP.CreateDefaultSubobject<UPaperFlipbookComponent>(this, TEXT("RenderComponent"));
+	RenderComponent = ObjectInitializer.CreateDefaultSubobject<UPaperFlipbookComponent>(this, TEXT("RenderComponent"));
 
 	RootComponent = RenderComponent;
 }
diff --git a/Engine/Plugins/2D/Paper2D/Source/Paper2D/Private/PaperFlipbookComponent.cpp b/Engine/Plugins/2D/Paper2D/Source/Paper2D/Private/PaperFlipbookComponent.cpp
--- a/Engine/Plugins/2D/Paper2D/Source/Paper2D/Private/PaperFlipbookComponent.cpp
+++ b/Engine/Plugins/2D/Paper2D/Source/Paper2D/Private/PaperFlipbookComponent.cpp
@@ -15,8 +15,8 @@
 //////////////////////////////////////////////////////////////////////////
 // UPaperFlipbookComponent
 
-UPaperFlipbookComponent::UPaperFlipbookComponent(const FPostConstructInitializeProperties& PCIP)
-	: Super(PCIP)
+UPaperFlipbookComponent::UPaperFlipbookComponent(const FObjectInitializer& ObjectInitializer)
+	: Super(ObjectInitializer)
 {
 	BodyInstance.SetCollisionEnabled(ECollisionEnabled::NoCollision);
 
diff --git a/Engine/Plugins/2D/Paper2D/Source/Paper2D/Private/PaperRuntimeSettings.cpp b/Engine/Plugins/2D/Paper2D/Source/Paper2D/Private/PaperRuntimeSettings.cpp
--- a/Engine/Plugins/2D/Paper2D/Source/Paper2D/Private/PaperRuntimeSettings.cpp
+++ b/Engine/Plugins/2D/Paper2D/Source/Paper2D/Private/PaperRuntimeSettings.cpp
@@ -6,8 +6,8 @@
 //////////////////////////////////////////////////////////////////////////
 // UPaperRuntimeSettings
 
-UPaperRuntimeSettings::UPaperRuntimeSettings(const FPostConstructInitializeProperties& PCIP)
-	: Super(PCIP)
+UPaperRuntimeSettings::UPaperRuntimeSettings(const FObjectInitializer& ObjectInitializer)
+	: Super(ObjectInitializer)
 	, DefaultPixelsPerUnrealUnit(2.56f)
 	, bEnableSpriteAtlasGroups(false)
 	, bEnableTileMapEditing(false)
diff --git a/Engine/Plugins/2D/Paper2D/Source/Paper2D/Private/PaperSprite.cpp b/Engine/Plugins/2D/Paper2D/Source/Paper2D/Private/PaperSprite.cpp
--- a/Engine/Plugins/2D/Paper2D/Source/Paper2D/Private/PaperSprite.cpp
+++ b/Engine/Plugins/2D/Paper2D/Source/Paper2D/Private/PaperSprite.cpp
@@ -250,8 +250,8 @@ void FSpriteDrawCallRecord::BuildFromSprite(const UPaperSprite* Sprite)
 //////////////////////////////////////////////////////////////////////////
 // UPaperSprite
 
-UPaperSprite::UPaperSprite(const FPostConstructInitializeProperties& PCIP)
-	: Super(PCIP)
+UPaperSprite::UPaperSprite(const FObjectInitializer& ObjectInitializer)
+	: Super(ObjectInitializer)
 {
 	// Default to using physics
 	SpriteCollisionDomain = ESpriteCollisionMode::Use3DPhysics;
diff --git a/Engine/Plugins/2D/Paper2D/Source/Paper2D/Private/PaperSpriteActor.cpp b/Engine/Plugins/2D/Paper2D/Source/Paper2D/Private/PaperSpriteActor.cpp
--- a/Engine/Plugins/2D/Paper2D/Source/Paper2D/Private/PaperSpriteActor.cpp
+++ b/Engine/Plugins/2D/Paper2D/Source/Paper2D/Private/PaperSpriteActor.cpp
@@ -6,10 +6,10 @@
 //////////////////////////////////////////////////////////////////////////
 // APaperSpriteActor
 
-APaperSpriteActor::APaperSpriteActor(const FPostConstructInitializeProperties& PCIP)
-	: Super(PCIP)
+APaperSpriteActor::APaperSpriteActor(const FObjectInitializer& ObjectInitializer)
+	: Super(ObjectInitializer)
 {
-	RenderComponent = PCIP.CreateDefaultSubobject<UPaperSpriteComponent>(this, TEXT("RenderComponent"));
+	RenderComponent = ObjectInitializer.CreateDefaultSubobject<UPaperSpriteComponent>(this, TEXT("RenderComponent"));
 	RenderComponent->SetCollisionProfileName(UCollisionProfile::BlockAll_ProfileName);
 	RenderComponent->Mobility = EComponentMobility::Static;
 
diff --git a/Engine/Plugins/2D/Paper2D/Source/Paper2D/Private/PaperSpriteAtlas.cpp b/Engine/Plugins/2D/Paper2D/Source/Paper2D/Private/PaperSpriteAtlas.cpp
--- a/Engine/Plugins/2D/Paper2D/Source/Paper2D/Private/PaperSpriteAtlas.cpp
+++ b/Engine/Plugins/2D/Paper2D/Source/Paper2D/Private/PaperSpriteAtlas.cpp
@@ -5,8 +5,8 @@
 //////////////////////////////////////////////////////////////////////////
 // UPaperSpriteAtlas
 
-UPaperSpriteAtlas::UPaperSpriteAtlas(const FPostConstructInitializeProperties& PCIP)
-	: Super(PCIP)
+UPaperSpriteAtlas::UPaperSpriteAtlas(const FObjectInitializer& ObjectInitializer)
+	: Super(ObjectInitializer)
 #if WITH_EDITORONLY_DATA
 	, MaxWidth(2048)
 	, MaxHeight(2048)
diff --git a/Engine/Plugins/2D/Paper2D/Source/Paper2D/Private/PaperSpriteComponent.cpp b/Engine/Plugins/2D/Paper2D/Source/Paper2D/Private/PaperSpriteComponent.cpp
--- a/Engine/Plugins/2D/Paper2D/Source/Paper2D/Private/PaperSpriteComponent.cpp
+++ b/Engine/Plugins/2D/Paper2D/Source/Paper2D/Private/PaperSpriteComponent.cpp
@@ -15,8 +15,8 @@
 //////////////////////////////////////////////////////////////////////////
 // UPaperSpriteComponent
 
-UPaperSpriteComponent::UPaperSpriteComponent(const FPostConstructInitializeProperties& PCIP)
-	: Super(PCIP)
+UPaperSpriteComponent::UPaperSpriteComponent(const FObjectInitializer& ObjectInitializer)
+	: Super(ObjectInitializer)
 {
 	SetCollisionProfileName(UCollisionProfile::BlockAllDynamic_ProfileName);
 
diff --git a/Engine/Plugins/2D/Paper2D/Source/Paper2D/Private/PaperTileLayer.cpp b/Engine/Plugins/2D/Paper2D/Source/Paper2D/Private/PaperTileLayer.cpp
--- a/Engine/Plugins/2D/Paper2D/Source/Paper2D/Private/PaperTileLayer.cpp
+++ b/Engine/Plugins/2D/Paper2D/Source/Paper2D/Private/PaperTileLayer.cpp
@@ -6,8 +6,8 @@
 //////////////////////////////////////////////////////////////////////////
 // UPaperTileLayer
 
-UPaperTileLayer::UPaperTileLayer(const FPostConstructInitializeProperties& PCIP)
-	: Super(PCIP)
+UPaperTileLayer::UPaperTileLayer(const FObjectInitializer& ObjectInitializer)
+	: Super(ObjectInitializer)
 {
 	LayerWidth = 4;
 	LayerHeight = 4;
diff --git a/Engine/Plugins/2D/Paper2D/Source/Paper2D/Private/PaperTileMap.cpp b/Engine/Plugins/2D/Paper2D/Source/Paper2D/Private/PaperTileMap.cpp
--- a/Engine/Plugins/2D/Paper2D/Source/Paper2D/Private/PaperTileMap.cpp
+++ b/Engine/Plugins/2D/Paper2D/Source/Paper2D/Private/PaperTileMap.cpp
@@ -7,8 +7,8 @@
 //////////////////////////////////////////////////////////////////////////
 // UPaperTileMap
 
-UPaperTileMap::UPaperTileMap(const FPostConstructInitializeProperties& PCIP)
-	: Super(PCIP)
+UPaperTileMap::UPaperTileMap(const FObjectInitializer& ObjectInitializer)
+	: Super(ObjectInitializer)
 {
 	MapWidth = 4;
 	MapHeight = 4;
diff --git a/Engine/Plugins/2D/Paper2D/Source/Paper2D/Private/PaperTileMapActor.cpp b/Engine/Plugins/2D/Paper2D/Source/Paper2D/Private/PaperTileMapActor.cpp
--- a/Engine/Plugins/2D/Paper2D/Source/Paper2D/Private/PaperTileMapActor.cpp
+++ b/Engine/Plugins/2D/Paper2D/Source/Paper2D/Private/PaperTileMapActor.cpp
@@ -5,10 +5,10 @@
 //////////////////////////////////////////////////////////////////////////
 // APaperTileMapActor
 
-APaperTileMapActor::APaperTileMapActor(const FPostConstructInitializeProperties& PCIP)
-	: Super(PCIP)
+APaperTileMapActor::APaperTileMapActor(const FObjectInitializer& ObjectInitializer)
+	: Super(ObjectInitializer)
 {
-	RenderComponent = PCIP.CreateDefaultSubobject<UPaperTileMapRenderComponent>(this, TEXT("RenderComponent"));
+	RenderComponent = ObjectInitializer.CreateDefaultSubobject<UPaperTileMapRenderComponent>(this, TEXT("RenderComponent"));
 
 	RootComponent = RenderComponent;
 }
diff --git a/Engine/Plugins/2D/Paper2D/Source/Paper2D/Private/PaperTileMapRenderComponent.cpp b/Engine/Plugins/2D/Paper2D/Source/Paper2D/Private/PaperTileMapRenderComponent.cpp
--- a/Engine/Plugins/2D/Paper2D/Source/Paper2D/Private/PaperTileMapRenderComponent.cpp
+++ b/Engine/Plugins/2D/Paper2D/Source/Paper2D/Private/PaperTileMapRenderComponent.cpp
@@ -10,8 +10,8 @@
 //////////////////////////////////////////////////////////////////////////
 // UPaperTileMapRenderComponent
 
-UPaperTileMapRenderComponent::UPaperTileMapRenderComponent(const FPostConstructInitializeProperties& PCIP)
-	: Super(PCIP)
+UPaperTileMapRenderComponent::UPaperTileMapRenderComponent(const FObjectInitializer& ObjectInitializer)
+	: Super(ObjectInitializer)
 {
 	BodyInstance.SetCollisionEnabled(ECollisionEnabled::QueryAndPhysics);
 
diff --git a/Engine/Plugins/2D/Paper2D/Source/Paper2D/Private/PaperTileSet.cpp b/Engine/Plugins/2D/Paper2D/Source/Paper2D/Private/PaperTileSet.cpp
--- a/Engine/Plugins/2D/Paper2D/Source/Paper2D/Private/PaperTileSet.cpp
+++ b/Engine/Plugins/2D/Paper2D/Source/Paper2D/Private/PaperTileSet.cpp
@@ -5,8 +5,8 @@
 //////////////////////////////////////////////////////////////////////////
 // UPaperTileSet
 
-UPaperTileSet::UPaperTileSet(const FPostConstructInitializeProperties& PCIP)
-	: Super(PCIP)
+UPaperTileSet::UPaperTileSet(const FObjectInitializer& ObjectInitializer)
+	: Super(ObjectInitializer)
 {
 	TileWidth = 32;
 	TileHeight = 32;
diff --git a/Engine/Plugins/2D/Paper2D/Source/Paper2D/Private/Rendering/MaterialExpressionSpriteTextureSampler.cpp b/Engine/Plugins/2D/Paper2D/Source/Paper2D/Private/Rendering/MaterialExpressionSpriteTextureSampler.cpp
--- a/Engine/Plugins/2D/Paper2D/Source/Paper2D/Private/Rendering/MaterialExpressionSpriteTextureSampler.cpp
+++ b/Engine/Plugins/2D/Paper2D/Source/Paper2D/Private/Rendering/MaterialExpressionSpriteTextureSampler.cpp
@@ -8,8 +8,8 @@
 //////////////////////////////////////////////////////////////////////////
 // UMaterialExpressionSpriteTextureSampler
 
-UMaterialExpressionSpriteTextureSampler::UMaterialExpressionSpriteTextureSampler(const FPostConstructInitializeProperties& PCIP)
-	: Super(PCIP)
+UMaterialExpressionSpriteTextureSampler::UMaterialExpressionSpriteTextureSampler(const FObjectInitializer& ObjectInitializer)
+	: Super(ObjectInitializer)
 {
 	ParameterName = TEXT("SpriteTexture");
 }
diff --git a/Engine/Plugins/2D/Paper2D/Source/Paper2D/Private/Rendering/PaperBatchComponent.cpp b/Engine/Plugins/2D/Paper2D/Source/Paper2D/Private/Rendering/PaperBatchComponent.cpp
--- a/Engine/Plugins/2D/Paper2D/Source/Paper2D/Private/Rendering/PaperBatchComponent.cpp
+++ b/Engine/Plugins/2D/Paper2D/Source/Paper2D/Private/Rendering/PaperBatchComponent.cpp
@@ -7,8 +7,8 @@
 //////////////////////////////////////////////////////////////////////////
 // UPaperBatchComponent
 
-UPaperBatchComponent::UPaperBatchComponent(const FPostConstructInitializeProperties& PCIP)
-	: Super(PCIP)
+UPaperBatchComponent::UPaperBatchComponent(const FObjectInitializer& ObjectInitializer)
+	: Super(ObjectInitializer)
 {
 	BodyInstance.SetCollisionEnabled(ECollisionEnabled::NoCollision);
 
diff --git a/Engine/Plugins/2D/Paper2D/Source/Paper2D/Private/Terrain/PaperTerrainActor.cpp b/Engine/Plugins/2D/Paper2D/Source/Paper2D/Private/Terrain/PaperTerrainActor.cpp
--- a/Engine/Plugins/2D/Paper2D/Source/Paper2D/Private/Terrain/PaperTerrainActor.cpp
+++ b/Engine/Plugins/2D/Paper2D/Source/Paper2D/Private/Terrain/PaperTerrainActor.cpp
@@ -5,12 +5,12 @@
 //////////////////////////////////////////////////////////////////////////
 // APaperTerrainActor
 
-APaperTerrainActor::APaperTerrainActor(const FPostConstructInitializeProperties& PCIP)
-	: Super(PCIP)
+APaperTerrainActor::APaperTerrainActor(const FObjectInitializer& ObjectInitializer)
+	: Super(ObjectInitializer)
 {
-	DummyRoot = PCIP.CreateDefaultSubobject<USceneComponent>(this, TEXT("RootComponent"));
-	SplineComponent = PCIP.CreateDefaultSubobject<UPaperTerrainSplineComponent>(this, TEXT("SplineComponent"));
- 	RenderComponent = PCIP.CreateDefaultSubobject<UPaperTerrainComponent>(this, TEXT("RenderComponent"));
+	DummyRoot = ObjectInitializer.CreateDefaultSubobject<USceneComponent>(this, TEXT("RootComponent"));
+	SplineComponent = ObjectInitializer.CreateDefaultSubobject<UPaperTerrainSplineComponent>(this, TEXT("SplineComponent"));
+ 	RenderComponent = ObjectInitializer.CreateDefaultSubobject<UPaperTerrainComponent>(this, TEXT("RenderComponent"));
  
 	SplineComponent->AttachParent = DummyRoot;
 	RenderComponent->AttachParent = DummyRoot;
diff --git a/Engine/Plugins/2D/Paper2D/Source/Paper2D/Private/Terrain/PaperTerrainComponent.cpp b/Engine/Plugins/2D/Paper2D/Source/Paper2D/Private/Terrain/PaperTerrainComponent.cpp
--- a/Engine/Plugins/2D/Paper2D/Source/Paper2D/Private/Terrain/PaperTerrainComponent.cpp
+++ b/Engine/Plugins/2D/Paper2D/Source/Paper2D/Private/Terrain/PaperTerrainComponent.cpp
@@ -79,8 +79,8 @@ void FPaperTerrainSceneProxy::GetDynamicMeshElementsForView(const FSceneView* Vi
 //////////////////////////////////////////////////////////////////////////
 // UPaperTerrainComponent
 
-UPaperTerrainComponent::UPaperTerrainComponent(const FPostConstructInitializeProperties& PCIP)
-	: Super(PCIP)
+UPaperTerrainComponent::UPaperTerrainComponent(const FObjectInitializer& ObjectInitializer)
+	: Super(ObjectInitializer)
 	, TerrainColor(FLinearColor::White)
 	, ReparamStepsPerSegment(8)
 	, SpriteCollisionDomain(ESpriteCollisionMode::Use3DPhysics)
diff --git a/Engine/Plugins/2D/Paper2D/Source/Paper2D/Private/Terrain/PaperTerrainMaterial.cpp b/Engine/Plugins/2D/Paper2D/Source/Paper2D/Private/Terrain/PaperTerrainMaterial.cpp
--- a/Engine/Plugins/2D/Paper2D/Source/Paper2D/Private/Terrain/PaperTerrainMaterial.cpp
+++ b/Engine/Plugins/2D/Paper2D/Source/Paper2D/Private/Terrain/PaperTerrainMaterial.cpp
@@ -7,8 +7,8 @@
 //////////////////////////////////////////////////////////////////////////
 // UPaperTerrainMaterial
 
-UPaperTerrainMaterial::UPaperTerrainMaterial(const FPostConstructInitializeProperties& PCIP)
-	: Super(PCIP)
+UPaperTerrainMaterial::UPaperTerrainMaterial(const FObjectInitializer& ObjectInitializer)
+	: Super(ObjectInitializer)
 {
 	// Add a default rule
 	FPaperTerrainMaterialRule& DefaultRule = *new (Rules) FPaperTerrainMaterialRule();
diff --git a/Engine/Plugins/2D/Paper2D/Source/Paper2D/Private/Terrain/PaperTerrainSplineComponent.cpp b/Engine/Plugins/2D/Paper2D/Source/Paper2D/Private/Terrain/PaperTerrainSplineComponent.cpp
--- a/Engine/Plugins/2D/Paper2D/Source/Paper2D/Private/Terrain/PaperTerrainSplineComponent.cpp
+++ b/Engine/Plugins/2D/Paper2D/Source/Paper2D/Private/Terrain/PaperTerrainSplineComponent.cpp
@@ -5,8 +5,8 @@
 //////////////////////////////////////////////////////////////////////////
 // UPaperTerrainSplineComponent
 
-UPaperTerrainSplineComponent::UPaperTerrainSplineComponent(const FPostConstructInitializeProperties& PCIP)
-	: Super(PCIP)
+UPaperTerrainSplineComponent::UPaperTerrainSplineComponent(const FObjectInitializer& ObjectInitializer)
+	: Super(ObjectInitializer)
 {
 }
 
diff --git a/Engine/Plugins/2D/Paper2D/Source/Paper2DEditor/Private/Atlasing/PaperSpriteAtlasFactory.cpp b/Engine/Plugins/2D/Paper2D/Source/Paper2DEditor/Private/Atlasing/PaperSpriteAtlasFactory.cpp
--- a/Engine/Plugins/2D/Paper2D/Source/Paper2DEditor/Private/Atlasing/PaperSpriteAtlasFactory.cpp
+++ b/Engine/Plugins/2D/Paper2D/Source/Paper2DEditor/Private/Atlasing/PaperSpriteAtlasFactory.cpp
@@ -6,8 +6,8 @@
 //////////////////////////////////////////////////////////////////////////
 // UPaperSpriteAtlasFactory
 
-UPaperSpriteAtlasFactory::UPaperSpriteAtlasFactory(const class FPostConstructInitializeProperties& PCIP)
-	: Super(PCIP)
+UPaperSpriteAtlasFactory::UPaperSpriteAtlasFactory(const FObjectInitializer& ObjectInitializer)
+	: Super(ObjectInitializer)
 {
 	SupportedClass = UPaperSpriteAtlas::StaticClass();
 
diff --git a/Engine/Plugins/2D/Paper2D/Source/Paper2DEditor/Private/PaperAnimatedSpriteActorFactory.cpp b/Engine/Plugins/2D/Paper2D/Source/Paper2DEditor/Private/PaperAnimatedSpriteActorFactory.cpp
--- a/Engine/Plugins/2D/Paper2D/Source/Paper2DEditor/Private/PaperAnimatedSpriteActorFactory.cpp
+++ b/Engine/Plugins/2D/Paper2D/Source/Paper2DEditor/Private/PaperAnimatedSpriteActorFactory.cpp
@@ -6,8 +6,8 @@
 //////////////////////////////////////////////////////////////////////////
 // UPaperFlipbookActorFactory
 
-UPaperFlipbookActorFactory::UPaperFlipbookActorFactory(const FPostConstructInitializeProperties& PCIP)
-	: Super(PCIP)
+UPaperFlipbookActorFactory::UPaperFlipbookActorFactory(const FObjectInitializer& ObjectInitializer)
+	: Super(ObjectInitializer)
 {
 	DisplayName = NSLOCTEXT("Paper2D", "PaperFlipbookFactoryDisplayName", "Add Animated Sprite");
 	NewActorClass = APaperFlipbookActor::StaticClass();
diff --git a/Engine/Plugins/2D/Paper2D/Source/Paper2DEditor/Private/PaperAnimatedSpriteThumbnailRenderer.cpp b/Engine/Plugins/2D/Paper2D/Source/Paper2DEditor/Private/PaperAnimatedSpriteThumbnailRenderer.cpp
--- a/Engine/Plugins/2D/Paper2D/Source/Paper2DEditor/Private/PaperAnimatedSpriteThumbnailRenderer.cpp
+++ b/Engine/Plugins/2D/Paper2D/Source/Paper2DEditor/Private/PaperAnimatedSpriteThumbnailRenderer.cpp
@@ -5,8 +5,8 @@
 //////////////////////////////////////////////////////////////////////////
 // UPaperFlipbookThumbnailRenderer
 
-UPaperFlipbookThumbnailRenderer::UPaperFlipbookThumbnailRenderer(const FPostConstructInitializeProperties& PCIP)
-	: Super(PCIP)
+UPaperFlipbookThumbnailRenderer::UPaperFlipbookThumbnailRenderer(const FObjectInitializer& ObjectInitializer)
+	: Super(ObjectInitializer)
 {
 }
 
diff --git a/Engine/Plugins/2D/Paper2D/Source/Paper2DEditor/Private/PaperFlipbookFactory.cpp b/Engine/Plugins/2D/Paper2D/Source/Paper2DEditor/Private/PaperFlipbookFactory.cpp
--- a/Engine/Plugins/2D/Paper2D/Source/Paper2DEditor/Private/PaperFlipbookFactory.cpp
+++ b/Engine/Plugins/2D/Paper2D/Source/Paper2DEditor/Private/PaperFlipbookFactory.cpp
@@ -7,8 +7,8 @@
 /////////////////////////////////////////////////////
 // UPaperFlipbookFactory
 
-UPaperFlipbookFactory::UPaperFlipbookFactory(const class FPostConstructInitializeProperties& PCIP)
-	: Super(PCIP)
+UPaperFlipbookFactory::UPaperFlipbookFactory(const FObjectInitializer& ObjectInitializer)
+	: Super(ObjectInitializer)
 {
 	bCreateNew = true;
 	bEditAfterNew = true;
diff --git a/Engine/Plugins/2D/Paper2D/Source/Paper2DEditor/Private/PaperSpriteActorFactory.cpp b/Engine/Plugins/2D/Paper2D/Source/Paper2DEditor/Private/PaperSpriteActorFactory.cpp
--- a/Engine/Plugins/2D/Paper2D/Source/Paper2DEditor/Private/PaperSpriteActorFactory.cpp
+++ b/Engine/Plugins/2D/Paper2D/Source/Paper2DEditor/Private/PaperSpriteActorFactory.cpp
@@ -6,8 +6,8 @@
 //////////////////////////////////////////////////////////////////////////
 // UPaperSpriteActorFactory
 
-UPaperSpriteActorFactory::UPaperSpriteActorFactory(const FPostConstructInitializeProperties& PCIP)
-	: Super(PCIP)
+UPaperSpriteActorFactory::UPaperSpriteActorFactory(const FObjectInitializer& ObjectInitializer)
+	: Super(ObjectInitializer)
 {
 	DisplayName = NSLOCTEXT("Paper2D", "PaperSpriteFactoryDisplayName", "Add Sprite");
 	NewActorClass = APaperSpriteActor::StaticClass();
diff --git a/Engine/Plugins/2D/Paper2D/Source/Paper2DEditor/Private/PaperSpriteFactory.cpp b/Engine/Plugins/2D/Paper2D/Source/Paper2DEditor/Private/PaperSpriteFactory.cpp
--- a/Engine/Plugins/2D/Paper2D/Source/Paper2DEditor/Private/PaperSpriteFactory.cpp
+++ b/Engine/Plugins/2D/Paper2D/Source/Paper2DEditor/Private/PaperSpriteFactory.cpp
@@ -7,8 +7,8 @@
 /////////////////////////////////////////////////////
 // UPaperSpriteFactory
 
-UPaperSpriteFactory::UPaperSpriteFactory(const class FPostConstructInitializeProperties& PCIP)
-	: Super(PCIP)
+UPaperSpriteFactory::UPaperSpriteFactory(const FObjectInitializer& ObjectInitializer)
+	: Super(ObjectInitializer)
 	, bUseSourceRegion(false)
 	, InitialSourceUV(0, 0)
 	, InitialSourceDimension(0, 0)
diff --git a/Engine/Plugins/2D/Paper2D/Source/Paper2DEditor/Private/PaperSpriteThumbnailRenderer.cpp b/Engine/Plugins/2D/Paper2D/Source/Paper2DEditor/Private/PaperSpriteThumbnailRenderer.cpp
--- a/Engine/Plugins/2D/Paper2D/Source/Paper2DEditor/Private/PaperSpriteThumbnailRenderer.cpp
+++ b/Engine/Plugins/2D/Paper2D/Source/Paper2DEditor/Private/PaperSpriteThumbnailRenderer.cpp
@@ -5,8 +5,8 @@
 //////////////////////////////////////////////////////////////////////////
 // UPaperSpriteThumbnailRenderer
 
-UPaperSpriteThumbnailRenderer::UPaperSpriteThumbnailRenderer(const FPostConstructInitializeProperties& PCIP)
-	: Super(PCIP)
+UPaperSpriteThumbnailRenderer::UPaperSpriteThumbnailRenderer(const FObjectInitializer& ObjectInitializer)
+	: Super(ObjectInitializer)
 {
 }
 
diff --git a/Engine/Plugins/2D/Paper2D/Source/Paper2DEditor/Private/PaperTileSetThumbnailRenderer.cpp b/Engine/Plugins/2D/Paper2D/Source/Paper2DEditor/Private/PaperTileSetThumbnailRenderer.cpp
--- a/Engine/Plugins/2D/Paper2D/Source/Paper2DEditor/Private/PaperTileSetThumbnailRenderer.cpp
+++ b/Engine/Plugins/2D/Paper2D/Source/Paper2DEditor/Private/PaperTileSetThumbnailRenderer.cpp
@@ -5,8 +5,8 @@
 //////////////////////////////////////////////////////////////////////////
 // UPaperTileSetThumbnailRenderer
 
-UPaperTileSetThumbnailRenderer::UPaperTileSetThumbnailRenderer(const FPostConstructInitializeProperties& PCIP)
-	: Super(PCIP)
+UPaperTileSetThumbnailRenderer::UPaperTileSetThumbnailRenderer(const FObjectInitializer& ObjectInitializer)
+	: Super(ObjectInitializer)
 {
 }
 
diff --git a/Engine/Plugins/2D/Paper2D/Source/Paper2DEditor/Private/SpriteEditor/PSE_DeleteTool.cpp b/Engine/Plugins/2D/Paper2D/Source/Paper2DEditor/Private/SpriteEditor/PSE_DeleteTool.cpp
--- a/Engine/Plugins/2D/Paper2D/Source/Paper2DEditor/Private/SpriteEditor/PSE_DeleteTool.cpp
+++ b/Engine/Plugins/2D/Paper2D/Source/Paper2DEditor/Private/SpriteEditor/PSE_DeleteTool.cpp
@@ -5,8 +5,8 @@
 //////////////////////////////////////////////////////////////////////////
 // UPSE_DeleteTool
 
-UPSE_DeleteTool::UPSE_DeleteTool(const FPostConstructInitializeProperties& PCIP)
-	: Super(PCIP)
+UPSE_DeleteTool::UPSE_DeleteTool(const FObjectInitializer& ObjectInitializer)
+	: Super(ObjectInitializer)
 {
 }
 
diff --git a/Engine/Plugins/2D/Paper2D/Source/Paper2DEditor/Private/SpriteEditor/PSE_LineTool.cpp b/Engine/Plugins/2D/Paper2D/Source/Paper2DEditor/Private/SpriteEditor/PSE_LineTool.cpp
--- a/Engine/Plugins/2D/Paper2D/Source/Paper2DEditor/Private/SpriteEditor/PSE_LineTool.cpp
+++ b/Engine/Plugins/2D/Paper2D/Source/Paper2DEditor/Private/SpriteEditor/PSE_LineTool.cpp
@@ -5,8 +5,8 @@
 //////////////////////////////////////////////////////////////////////////
 // UPSE_LineTool
 
-UPSE_LineTool::UPSE_LineTool(const FPostConstructInitializeProperties& PCIP)
-	: Super(PCIP)
+UPSE_LineTool::UPSE_LineTool(const FObjectInitializer& ObjectInitializer)
+	: Super(ObjectInitializer)
 {
 }
 
diff --git a/Engine/Plugins/2D/Paper2D/Source/Paper2DEditor/Private/SpriteEditor/PSE_SplitEdgeTool.cpp b/Engine/Plugins/2D/Paper2D/Source/Paper2DEditor/Private/SpriteEditor/PSE_SplitEdgeTool.cpp
--- a/Engine/Plugins/2D/Paper2D/Source/Paper2DEditor/Private/SpriteEditor/PSE_SplitEdgeTool.cpp
+++ b/Engine/Plugins/2D/Paper2D/Source/Paper2DEditor/Private/SpriteEditor/PSE_SplitEdgeTool.cpp
@@ -5,7 +5,7 @@
 //////////////////////////////////////////////////////////////////////////
 // UPSE_SplitEdgeTool
 
-UPSE_SplitEdgeTool::UPSE_SplitEdgeTool(const FPostConstructInitializeProperties& PCIP)
-	: Super(PCIP)
+UPSE_SplitEdgeTool::UPSE_SplitEdgeTool(const FObjectInitializer& ObjectInitializer)
+	: Super(ObjectInitializer)
 {
 }
diff --git a/Engine/Plugins/2D/Paper2D/Source/Paper2DEditor/Private/SpriteEditor/SpriteEditorTool.cpp b/Engine/Plugins/2D/Paper2D/Source/Paper2DEditor/Private/SpriteEditor/SpriteEditorTool.cpp
--- a/Engine/Plugins/2D/Paper2D/Source/Paper2DEditor/Private/SpriteEditor/SpriteEditorTool.cpp
+++ b/Engine/Plugins/2D/Paper2D/Source/Paper2DEditor/Private/SpriteEditor/SpriteEditorTool.cpp
@@ -5,8 +5,8 @@
 //////////////////////////////////////////////////////////////////////////
 // USpriteEditorTool
 
-USpriteEditorTool::USpriteEditorTool(const FPostConstructInitializeProperties& PCIP)
-	: Super(PCIP)
+USpriteEditorTool::USpriteEditorTool(const FObjectInitializer& ObjectInitializer)
+	: Super(ObjectInitializer)
 {
 }
 
diff --git a/Engine/Plugins/2D/Paper2D/Source/Paper2DEditor/Private/TerrainSplineEditing/TerrainSplineActorFactory.cpp b/Engine/Plugins/2D/Paper2D/Source/Paper2DEditor/Private/TerrainSplineEditing/TerrainSplineActorFactory.cpp
--- a/Engine/Plugins/2D/Paper2D/Source/Paper2DEditor/Private/TerrainSplineEditing/TerrainSplineActorFactory.cpp
+++ b/Engine/Plugins/2D/Paper2D/Source/Paper2DEditor/Private/TerrainSplineEditing/TerrainSplineActorFactory.cpp
@@ -6,8 +6,8 @@
 //////////////////////////////////////////////////////////////////////////
 // UTerrainSplineActorFactory
 
-UTerrainSplineActorFactory::UTerrainSplineActorFactory(const FPostConstructInitializeProperties& PCIP)
-	: Super(PCIP)
+UTerrainSplineActorFactory::UTerrainSplineActorFactory(const FObjectInitializer& ObjectInitializer)
+	: Super(ObjectInitializer)
 {
 	DisplayName = NSLOCTEXT("Paper2D", "TerrainSplineFactoryDisplayName", "Paper2D Terrain Spline");
 	NewActorClass = APaperTerrainActor::StaticClass();
diff --git a/Engine/Plugins/2D/Paper2D/Source/Paper2DEditor/Private/TileMapEditing/TileMapActorFactory.cpp b/Engine/Plugins/2D/Paper2D/Source/Paper2DEditor/Private/TileMapEditing/TileMapActorFactory.cpp
--- a/Engine/Plugins/2D/Paper2D/Source/Paper2DEditor/Private/TileMapEditing/TileMapActorFactory.cpp
+++ b/Engine/Plugins/2D/Paper2D/Source/Paper2DEditor/Private/TileMapEditing/TileMapActorFactory.cpp
@@ -6,8 +6,8 @@
 //////////////////////////////////////////////////////////////////////////
 // UTileMapActorFactory
 
-UTileMapActorFactory::UTileMapActorFactory(const FPostConstructInitializeProperties& PCIP)
-	: Super(PCIP)
+UTileMapActorFactory::UTileMapActorFactory(const FObjectInitializer& ObjectInitializer)
+	: Super(ObjectInitializer)
 {
 	DisplayName = NSLOCTEXT("Paper2D", "TileMapFactoryDisplayName", "Paper2D Tile Map");
 	NewActorClass = APaperTileMapActor::StaticClass();
diff --git a/Engine/Plugins/2D/Paper2D/Source/Paper2DEditor/Private/TileMapEditing/TileMapAssetImportData.cpp b/Engine/Plugins/2D/Paper2D/Source/Paper2DEditor/Private/TileMapEditing/TileMapAssetImportData.cpp
--- a/Engine/Plugins/2D/Paper2D/Source/Paper2DEditor/Private/TileMapEditing/TileMapAssetImportData.cpp
+++ b/Engine/Plugins/2D/Paper2D/Source/Paper2DEditor/Private/TileMapEditing/TileMapAssetImportData.cpp
@@ -6,8 +6,8 @@
 //////////////////////////////////////////////////////////////////////////
 // UTileMapAssetImportData
 
-UTileMapAssetImportData::UTileMapAssetImportData(const FPostConstructInitializeProperties& PCIP)
-	: Super(PCIP)
+UTileMapAssetImportData::UTileMapAssetImportData(const FObjectInitializer& ObjectInitializer)
+	: Super(ObjectInitializer)
 {
 }
 
diff --git a/Engine/Plugins/2D/Paper2D/Source/PaperJsonImporter/Private/PaperJsonImporterFactory.cpp b/Engine/Plugins/2D/Paper2D/Source/PaperJsonImporter/Private/PaperJsonImporterFactory.cpp
--- a/Engine/Plugins/2D/Paper2D/Source/PaperJsonImporter/Private/PaperJsonImporterFactory.cpp
+++ b/Engine/Plugins/2D/Paper2D/Source/PaperJsonImporter/Private/PaperJsonImporterFactory.cpp
@@ -14,8 +14,8 @@
 //////////////////////////////////////////////////////////////////////////
 // UPaperJsonImporterFactory
 
-UPaperJsonImporterFactory::UPaperJsonImporterFactory(const FPostConstructInitializeProperties& PCIP)
-	: Super(PCIP)
+UPaperJsonImporterFactory::UPaperJsonImporterFactory(const FObjectInitializer& ObjectInitializer)
+	: Super(ObjectInitializer)
 {
 	bCreateNew = false;
 	//bEditAfterNew = true;
diff --git a/Engine/Plugins/2D/Paper2D/Source/PaperJsonImporter/Private/PaperSpriteSheet.cpp b/Engine/Plugins/2D/Paper2D/Source/PaperJsonImporter/Private/PaperSpriteSheet.cpp
--- a/Engine/Plugins/2D/Paper2D/Source/PaperJsonImporter/Private/PaperSpriteSheet.cpp
+++ b/Engine/Plugins/2D/Paper2D/Source/PaperJsonImporter/Private/PaperSpriteSheet.cpp
@@ -3,8 +3,8 @@
 #include "PaperJsonImporterPrivatePCH.h"
 #include "PaperSpriteSheet.h"
 
-UPaperSpriteSheet::UPaperSpriteSheet(const FPostConstructInitializeProperties& PCIP)
-	: Super(PCIP)
+UPaperSpriteSheet::UPaperSpriteSheet(const FObjectInitializer& ObjectInitializer)
+	: Super(ObjectInitializer)
 {
 }
 
diff --git a/Engine/Plugins/2D/Paper2D/Source/PaperJsonImporter/Private/ReimportPaperJsonImporterFactory.cpp b/Engine/Plugins/2D/Paper2D/Source/PaperJsonImporter/Private/ReimportPaperJsonImporterFactory.cpp
--- a/Engine/Plugins/2D/Paper2D/Source/PaperJsonImporter/Private/ReimportPaperJsonImporterFactory.cpp
+++ b/Engine/Plugins/2D/Paper2D/Source/PaperJsonImporter/Private/ReimportPaperJsonImporterFactory.cpp
@@ -7,8 +7,8 @@
 #define LOCTEXT_NAMESPACE "PaperJsonImporter"
 
 
-UReimportPaperJsonImporterFactory::UReimportPaperJsonImporterFactory(const class FPostConstructInitializeProperties& PCIP)
-	: Super(PCIP)
+UReimportPaperJsonImporterFactory::UReimportPaperJsonImporterFactory(const FObjectInitializer& ObjectInitializer)
+	: Super(ObjectInitializer)
 {
 	SupportedClass = UPaperSpriteSheet::StaticClass();
 
diff --git a/Engine/Plugins/2D/Paper2D/Source/PaperTiledImporter/Private/PaperTiledImporterFactory.cpp b/Engine/Plugins/2D/Paper2D/Source/PaperTiledImporter/Private/PaperTiledImporterFactory.cpp
--- a/Engine/Plugins/2D/Paper2D/Source/PaperTiledImporter/Private/PaperTiledImporterFactory.cpp
+++ b/Engine/Plugins/2D/Paper2D/Source/PaperTiledImporter/Private/PaperTiledImporterFactory.cpp
@@ -71,8 +71,8 @@ bool ParseIntegerFields(FRequiredIntField* IntFieldArray, int32 ArrayCount, TSha
 //////////////////////////////////////////////////////////////////////////
 // UPaperTiledImporterFactory
 
-UPaperTiledImporterFactory::UPaperTiledImporterFactory(const FPostConstructInitializeProperties& PCIP)
-	: Super(PCIP)
+UPaperTiledImporterFactory::UPaperTiledImporterFactory(const FObjectInitializer& ObjectInitializer)
+	: Super(ObjectInitializer)
 {
 	bCreateNew = false;
 	//bEditAfterNew = true;
diff --git a/Engine/Plugins/Developer/UObjectPlugin/Source/UObjectPlugin/Private/MyPluginObject.cpp b/Engine/Plugins/Developer/UObjectPlugin/Source/UObjectPlugin/Private/MyPluginObject.cpp
--- a/Engine/Plugins/Developer/UObjectPlugin/Source/UObjectPlugin/Private/MyPluginObject.cpp
+++ b/Engine/Plugins/Developer/UObjectPlugin/Source/UObjectPlugin/Private/MyPluginObject.cpp
@@ -4,7 +4,7 @@
 #include "MyPluginObject.h"
 
 
-UMyPluginObject::UMyPluginObject( const FPostConstructInitializeProperties& PCIP )
-	: Super( PCIP )
+UMyPluginObject::UMyPluginObject( const FObjectInitializer& ObjectInitializer )
+	: Super( ObjectInitializer )
 {
 }
\ No newline at end of file
diff --git a/Engine/Plugins/Editor/MacGraphicsSwitching/Source/MacGraphicsSwitching/Private/MacGraphicsSwitchingSettings.cpp b/Engine/Plugins/Editor/MacGraphicsSwitching/Source/MacGraphicsSwitching/Private/MacGraphicsSwitchingSettings.cpp
--- a/Engine/Plugins/Editor/MacGraphicsSwitching/Source/MacGraphicsSwitching/Private/MacGraphicsSwitchingSettings.cpp
+++ b/Engine/Plugins/Editor/MacGraphicsSwitching/Source/MacGraphicsSwitching/Private/MacGraphicsSwitchingSettings.cpp
@@ -3,7 +3,7 @@
 #include "MacGraphicsSwitchingModule.h"
 #include "MacGraphicsSwitchingSettings.h"
 
-UMacGraphicsSwitchingSettings::UMacGraphicsSwitchingSettings( const class FPostConstructInitializeProperties& PCIP )
-	: Super(PCIP)
+UMacGraphicsSwitchingSettings::UMacGraphicsSwitchingSettings( const FObjectInitializer& ObjectInitializer )
+	: Super(ObjectInitializer)
 {
 }
diff --git a/Engine/Plugins/Editor/SpeedTreeImporter/Source/SpeedTreeImporter/Private/ReimportSpeedTreeFactory.cpp b/Engine/Plugins/Editor/SpeedTreeImporter/Source/SpeedTreeImporter/Private/ReimportSpeedTreeFactory.cpp
--- a/Engine/Plugins/Editor/SpeedTreeImporter/Source/SpeedTreeImporter/Private/ReimportSpeedTreeFactory.cpp
+++ b/Engine/Plugins/Editor/SpeedTreeImporter/Source/SpeedTreeImporter/Private/ReimportSpeedTreeFactory.cpp
@@ -10,8 +10,8 @@ DEFINE_LOG_CATEGORY_STATIC(LogEditorFactories, Log, All);
 
 
 
-UReimportSpeedTreeFactory::UReimportSpeedTreeFactory(const class FPostConstructInitializeProperties& PCIP)
-	: Super(PCIP)
+UReimportSpeedTreeFactory::UReimportSpeedTreeFactory(const FObjectInitializer& ObjectInitializer)
+	: Super(ObjectInitializer)
 {
 	#if WITH_SPEEDTREE
 		SupportedClass = UStaticMesh::StaticClass();
diff --git a/Engine/Plugins/Editor/SpeedTreeImporter/Source/SpeedTreeImporter/Private/SpeedTreeImportFactory.cpp b/Engine/Plugins/Editor/SpeedTreeImporter/Source/SpeedTreeImporter/Private/SpeedTreeImportFactory.cpp
--- a/Engine/Plugins/Editor/SpeedTreeImporter/Source/SpeedTreeImporter/Private/SpeedTreeImportFactory.cpp
+++ b/Engine/Plugins/Editor/SpeedTreeImporter/Source/SpeedTreeImporter/Private/SpeedTreeImportFactory.cpp
@@ -335,8 +335,8 @@ public:
 
 //////////////////////////////////////////////////////////////////////////
 
-USpeedTreeImportFactory::USpeedTreeImportFactory(const class FPostConstructInitializeProperties& PCIP)
-	: Super(PCIP)
+USpeedTreeImportFactory::USpeedTreeImportFactory(const FObjectInitializer& ObjectInitializer)
+	: Super(ObjectInitializer)
 {
 	bCreateNew = false;
 	bEditAfterNew = true;
@@ -402,7 +402,7 @@ static UTexture* CreateSpeedTreeMaterialTexture(UObject* Parent, FString Filenam
 	}
 	else
 	{
-		UTextureFactory* TextureFact = new UTextureFactory(FPostConstructInitializeProperties());
+		UTextureFactory* TextureFact = new UTextureFactory(FObjectInitializer());
 		TextureFact->SuppressImportOverwriteDialog();
 
 		if (bNormalMap)
@@ -555,7 +555,7 @@ static UMaterialInterface* CreateSpeedTreeMaterial(UObject* Parent, FString Mate
 	}
 	
 	// create an unreal material asset
-	UMaterialFactoryNew* MaterialFactory = new UMaterialFactoryNew(FPostConstructInitializeProperties());
+	UMaterialFactoryNew* MaterialFactory = new UMaterialFactoryNew(FObjectInitializer());
 	UMaterial* UnrealMaterial = (UMaterial*)MaterialFactory->FactoryCreateNew(UMaterial::StaticClass(), Package, *FixedMaterialName, RF_Standalone|RF_Public, NULL, GWarn);
 	if (UnrealMaterial != NULL)
 	{
@@ -999,7 +999,7 @@ UObject* USpeedTreeImportFactory::FactoryCreateBinary(UClass* InClass, UObject*
 					ExistingMesh->PreEditChange(NULL);
 				}
 				
-				StaticMesh = new(Package, FName(*MeshName), Flags | RF_Public) UStaticMesh(FPostConstructInitializeProperties());
+				StaticMesh = new(Package, FName(*MeshName), Flags | RF_Public) UStaticMesh(FObjectInitializer());
 
 				// @todo AssetImportData make a data class for speed tree assets
 				StaticMesh->AssetImportData = ConstructObject<UAssetImportData>(UAssetImportData::StaticClass(), StaticMesh);
diff --git a/Engine/Plugins/Experimental/LiveEditor/Source/LiveEditor/Classes/LiveEditorTypes.h b/Engine/Plugins/Experimental/LiveEditor/Source/LiveEditor/Classes/LiveEditorTypes.h
--- a/Engine/Plugins/Experimental/LiveEditor/Source/LiveEditor/Classes/LiveEditorTypes.h
+++ b/Engine/Plugins/Experimental/LiveEditor/Source/LiveEditor/Classes/LiveEditorTypes.h
@@ -70,5 +70,5 @@ class ULiveEditorTypes : public UObject
 {
 	GENERATED_UCLASS_BODY()
 
-	ULiveEditorTypes(const class FPostConstructInitializeProperties& PCIP) : Super(PCIP) {}
+	ULiveEditorTypes(const FObjectInitializer& ObjectInitializer) : Super(ObjectInitializer) {}
 };
diff --git a/Engine/Plugins/Experimental/LiveEditor/Source/LiveEditor/Private/K2Node_LiveEditObject.cpp b/Engine/Plugins/Experimental/LiveEditor/Source/LiveEditor/Private/K2Node_LiveEditObject.cpp
--- a/Engine/Plugins/Experimental/LiveEditor/Source/LiveEditor/Private/K2Node_LiveEditObject.cpp
+++ b/Engine/Plugins/Experimental/LiveEditor/Source/LiveEditor/Private/K2Node_LiveEditObject.cpp
@@ -138,8 +138,8 @@ namespace UK2Node_LiveEditObjectStatics
 FString UK2Node_LiveEditObject::LiveEditableVarPinCategory(TEXT("LiveEditableVar"));
 
 
-UK2Node_LiveEditObject::UK2Node_LiveEditObject(const class FPostConstructInitializeProperties& PCIP)
-	: Super(PCIP)
+UK2Node_LiveEditObject::UK2Node_LiveEditObject(const FObjectInitializer& ObjectInitializer)
+	: Super(ObjectInitializer)
 {
 }
 
diff --git a/Engine/Plugins/Experimental/LiveEditor/Source/LiveEditor/Private/LiveEditorBlueprint.cpp b/Engine/Plugins/Experimental/LiveEditor/Source/LiveEditor/Private/LiveEditorBlueprint.cpp
--- a/Engine/Plugins/Experimental/LiveEditor/Source/LiveEditor/Private/LiveEditorBlueprint.cpp
+++ b/Engine/Plugins/Experimental/LiveEditor/Source/LiveEditor/Private/LiveEditorBlueprint.cpp
@@ -17,8 +17,8 @@ namespace nLiveEditorBlueprint
 	}
 }
 
-ULiveEditorBlueprint::ULiveEditorBlueprint(const FPostConstructInitializeProperties& PCIP)
-	: Super(PCIP)
+ULiveEditorBlueprint::ULiveEditorBlueprint(const FObjectInitializer& ObjectInitializer)
+	: Super(ObjectInitializer)
 {
 }
 
diff --git a/Engine/Plugins/Experimental/LiveEditor/Source/LiveEditor/Private/LiveEditorBroadcaster.cpp b/Engine/Plugins/Experimental/LiveEditor/Source/LiveEditor/Private/LiveEditorBroadcaster.cpp
--- a/Engine/Plugins/Experimental/LiveEditor/Source/LiveEditor/Private/LiveEditorBroadcaster.cpp
+++ b/Engine/Plugins/Experimental/LiveEditor/Source/LiveEditor/Private/LiveEditorBroadcaster.cpp
@@ -1,7 +1,7 @@
 
 #include "LiveEditorPrivatePCH.h"
 
-ULiveEditorBroadcaster::ULiveEditorBroadcaster(const class FPostConstructInitializeProperties& PCIP)
-	: Super(PCIP)
+ULiveEditorBroadcaster::ULiveEditorBroadcaster(const FObjectInitializer& ObjectInitializer)
+	: Super(ObjectInitializer)
 {
 }
diff --git a/Engine/Plugins/Experimental/LiveEditor/Source/LiveEditor/Private/LiveEditorKismetLibrary.cpp b/Engine/Plugins/Experimental/LiveEditor/Source/LiveEditor/Private/LiveEditorKismetLibrary.cpp
--- a/Engine/Plugins/Experimental/LiveEditor/Source/LiveEditor/Private/LiveEditorKismetLibrary.cpp
+++ b/Engine/Plugins/Experimental/LiveEditor/Source/LiveEditor/Private/LiveEditorKismetLibrary.cpp
@@ -326,8 +326,8 @@ namespace nLiveEditorKismetLibrary
 /**
  * ULiveEditorKismetLibrary
  **/
-ULiveEditorKismetLibrary::ULiveEditorKismetLibrary(const class FPostConstructInitializeProperties& PCIP)
-	: Super(PCIP)
+ULiveEditorKismetLibrary::ULiveEditorKismetLibrary(const FObjectInitializer& ObjectInitializer)
+	: Super(ObjectInitializer)
 {
 }
 
diff --git a/Engine/Plugins/Experimental/LiveEditor/Source/LiveEditor/Private/LiveEditorManager.cpp b/Engine/Plugins/Experimental/LiveEditor/Source/LiveEditor/Private/LiveEditorManager.cpp
--- a/Engine/Plugins/Experimental/LiveEditor/Source/LiveEditor/Private/LiveEditorManager.cpp
+++ b/Engine/Plugins/Experimental/LiveEditor/Source/LiveEditor/Private/LiveEditorManager.cpp
@@ -606,7 +606,7 @@ void FLiveEditorManager::InjectNewBlueprintEditor( const TWeakPtr<FBlueprintEdit
 
 void FLiveEditorManager::CreateLiveEditorWorld()
 {
-	LiveEditorWorld = new UWorld(FPostConstructInitializeProperties(),FURL(NULL));
+	LiveEditorWorld = new UWorld(FObjectInitializer(),FURL(NULL));
 	LiveEditorWorld->WorldType = EWorldType::Preview;
 	
 	bool bTransactional = false;
diff --git a/Engine/Plugins/Experimental/Phya/Source/Phya/Private/PhyaCollisionHandler.cpp b/Engine/Plugins/Experimental/Phya/Source/Phya/Private/PhyaCollisionHandler.cpp
--- a/Engine/Plugins/Experimental/Phya/Source/Phya/Private/PhyaCollisionHandler.cpp
+++ b/Engine/Plugins/Experimental/Phya/Source/Phya/Private/PhyaCollisionHandler.cpp
@@ -28,8 +28,8 @@ uint32 GetTypeHash( const FPhyaBodyInstancePair& Pair )
 
 //////////////////////////////////////////////////////////////////////////
 
-UPhyaCollisionHandler::UPhyaCollisionHandler(const class FPostConstructInitializeProperties& PCIP)
-	: Super(PCIP)
+UPhyaCollisionHandler::UPhyaCollisionHandler(const FObjectInitializer& ObjectInitializer)
+	: Super(ObjectInitializer)
 {
 }
 
diff --git a/Engine/Plugins/Media/MediaPlayerEditor/Source/MediaPlayerEditor/Private/Factories/MediaPlayerFactory.cpp b/Engine/Plugins/Media/MediaPlayerEditor/Source/MediaPlayerEditor/Private/Factories/MediaPlayerFactory.cpp
--- a/Engine/Plugins/Media/MediaPlayerEditor/Source/MediaPlayerEditor/Private/Factories/MediaPlayerFactory.cpp
+++ b/Engine/Plugins/Media/MediaPlayerEditor/Source/MediaPlayerEditor/Private/Factories/MediaPlayerFactory.cpp
@@ -6,8 +6,8 @@
 /* UMediaPlayerFactory structors
  *****************************************************************************/
 
-UMediaPlayerFactory::UMediaPlayerFactory( const class FPostConstructInitializeProperties& PCIP )
-	: Super(PCIP)
+UMediaPlayerFactory::UMediaPlayerFactory( const FObjectInitializer& ObjectInitializer )
+	: Super(ObjectInitializer)
 {
 	SupportedClass = UMediaPlayer::StaticClass();
 
diff --git a/Engine/Plugins/Media/MediaPlayerEditor/Source/MediaPlayerEditor/Private/Factories/MediaPlayerFactoryNew.cpp b/Engine/Plugins/Media/MediaPlayerEditor/Source/MediaPlayerEditor/Private/Factories/MediaPlayerFactoryNew.cpp
--- a/Engine/Plugins/Media/MediaPlayerEditor/Source/MediaPlayerEditor/Private/Factories/MediaPlayerFactoryNew.cpp
+++ b/Engine/Plugins/Media/MediaPlayerEditor/Source/MediaPlayerEditor/Private/Factories/MediaPlayerFactoryNew.cpp
@@ -6,8 +6,8 @@
 /* UMediaPlayerFactoryNew structors
  *****************************************************************************/
 
-UMediaPlayerFactoryNew::UMediaPlayerFactoryNew( const class FPostConstructInitializeProperties& PCIP )
-	: Super(PCIP)
+UMediaPlayerFactoryNew::UMediaPlayerFactoryNew( const FObjectInitializer& ObjectInitializer )
+	: Super(ObjectInitializer)
 {
 	SupportedClass = UMediaPlayer::StaticClass();
 	bCreateNew = true;
diff --git a/Engine/Plugins/Media/MediaPlayerEditor/Source/MediaPlayerEditor/Private/Factories/MediaSoundWaveFactoryNew.cpp b/Engine/Plugins/Media/MediaPlayerEditor/Source/MediaPlayerEditor/Private/Factories/MediaSoundWaveFactoryNew.cpp
--- a/Engine/Plugins/Media/MediaPlayerEditor/Source/MediaPlayerEditor/Private/Factories/MediaSoundWaveFactoryNew.cpp
+++ b/Engine/Plugins/Media/MediaPlayerEditor/Source/MediaPlayerEditor/Private/Factories/MediaSoundWaveFactoryNew.cpp
@@ -6,8 +6,8 @@
 /* UMediaSoundWaveFactoryNew structors
  *****************************************************************************/
 
-UMediaSoundWaveFactoryNew::UMediaSoundWaveFactoryNew( const class FPostConstructInitializeProperties& PCIP )
-	: Super(PCIP)
+UMediaSoundWaveFactoryNew::UMediaSoundWaveFactoryNew( const FObjectInitializer& ObjectInitializer )
+	: Super(ObjectInitializer)
 {
 	SupportedClass = UMediaSoundWave::StaticClass();
 	bCreateNew = true;
diff --git a/Engine/Plugins/Media/MediaPlayerEditor/Source/MediaPlayerEditor/Private/Factories/MediaTextureFactoryNew.cpp b/Engine/Plugins/Media/MediaPlayerEditor/Source/MediaPlayerEditor/Private/Factories/MediaTextureFactoryNew.cpp
--- a/Engine/Plugins/Media/MediaPlayerEditor/Source/MediaPlayerEditor/Private/Factories/MediaTextureFactoryNew.cpp
+++ b/Engine/Plugins/Media/MediaPlayerEditor/Source/MediaPlayerEditor/Private/Factories/MediaTextureFactoryNew.cpp
@@ -6,8 +6,8 @@
 /* UMediaTextureFactoryNew structors
  *****************************************************************************/
 
-UMediaTextureFactoryNew::UMediaTextureFactoryNew( const class FPostConstructInitializeProperties& PCIP )
-	: Super(PCIP)
+UMediaTextureFactoryNew::UMediaTextureFactoryNew( const FObjectInitializer& ObjectInitializer )
+	: Super(ObjectInitializer)
 {
 	SupportedClass = UMediaTexture::StaticClass();
 	bCreateNew = true;
diff --git a/Engine/Plugins/Messaging/UdpMessaging/Source/UdpMessaging/Private/UdpMessagingClasses.cpp b/Engine/Plugins/Messaging/UdpMessaging/Source/UdpMessaging/Private/UdpMessagingClasses.cpp
--- a/Engine/Plugins/Messaging/UdpMessaging/Source/UdpMessaging/Private/UdpMessagingClasses.cpp
+++ b/Engine/Plugins/Messaging/UdpMessaging/Source/UdpMessaging/Private/UdpMessagingClasses.cpp
@@ -3,8 +3,8 @@
 #include "UdpMessagingPrivatePCH.h"
 
 
-UUdpMessagingSettings::UUdpMessagingSettings( const class FPostConstructInitializeProperties& PCIP )
-	: Super(PCIP)
+UUdpMessagingSettings::UUdpMessagingSettings( const FObjectInitializer& ObjectInitializer )
+	: Super(ObjectInitializer)
 	, EnableTransport(true)
 	, MulticastTimeToLive(1)
 	, EnableTunnel(false)
diff --git a/Engine/Plugins/Runtime/Analytics/AnalyticsBlueprintLibrary/Source/AnalyticsBlueprintLibrary/Private/AnalyticsBlueprintLibrary.cpp b/Engine/Plugins/Runtime/Analytics/AnalyticsBlueprintLibrary/Source/AnalyticsBlueprintLibrary/Private/AnalyticsBlueprintLibrary.cpp
--- a/Engine/Plugins/Runtime/Analytics/AnalyticsBlueprintLibrary/Source/AnalyticsBlueprintLibrary/Private/AnalyticsBlueprintLibrary.cpp
+++ b/Engine/Plugins/Runtime/Analytics/AnalyticsBlueprintLibrary/Source/AnalyticsBlueprintLibrary/Private/AnalyticsBlueprintLibrary.cpp
@@ -20,8 +20,8 @@ static inline TArray<FAnalyticsEventAttribute> ConvertAttrs(const TArray<FAnalyt
 }
 
 
-UAnalyticsBlueprintLibrary::UAnalyticsBlueprintLibrary(const class FPostConstructInitializeProperties& PCIP)
-	: Super(PCIP)
+UAnalyticsBlueprintLibrary::UAnalyticsBlueprintLibrary(const FObjectInitializer& ObjectInitializer)
+	: Super(ObjectInitializer)
 {
 }
 
diff --git a/Engine/Plugins/Runtime/CableComponent/Source/CableComponent/Private/CableActor.cpp b/Engine/Plugins/Runtime/CableComponent/Source/CableComponent/Private/CableActor.cpp
--- a/Engine/Plugins/Runtime/CableComponent/Source/CableComponent/Private/CableActor.cpp
+++ b/Engine/Plugins/Runtime/CableComponent/Source/CableComponent/Private/CableActor.cpp
@@ -3,9 +3,9 @@
 #include "CableComponentPluginPrivatePCH.h"
 
 
-ACableActor::ACableActor(const class FPostConstructInitializeProperties& PCIP)
-	: Super(PCIP)
+ACableActor::ACableActor(const FObjectInitializer& ObjectInitializer)
+	: Super(ObjectInitializer)
 {
-	CableComponent = PCIP.CreateDefaultSubobject<UCableComponent>(this, TEXT("CableComponent0"));
+	CableComponent = ObjectInitializer.CreateDefaultSubobject<UCableComponent>(this, TEXT("CableComponent0"));
 	RootComponent = CableComponent;
 }
\ No newline at end of file
diff --git a/Engine/Plugins/Runtime/CableComponent/Source/CableComponent/Private/CableComponent.cpp b/Engine/Plugins/Runtime/CableComponent/Source/CableComponent/Private/CableComponent.cpp
--- a/Engine/Plugins/Runtime/CableComponent/Source/CableComponent/Private/CableComponent.cpp
+++ b/Engine/Plugins/Runtime/CableComponent/Source/CableComponent/Private/CableComponent.cpp
@@ -378,8 +378,8 @@ private:
 
 //////////////////////////////////////////////////////////////////////////
 
-UCableComponent::UCableComponent( const FPostConstructInitializeProperties& PCIP )
-	: Super( PCIP )
+UCableComponent::UCableComponent( const FObjectInitializer& ObjectInitializer )
+	: Super( ObjectInitializer )
 {
 	PrimaryComponentTick.bCanEverTick = true;
 	bTickInEditor = true;
diff --git a/Engine/Plugins/Runtime/CustomMeshComponent/Source/CustomMeshComponent/Private/CustomMeshComponent.cpp b/Engine/Plugins/Runtime/CustomMeshComponent/Source/CustomMeshComponent/Private/CustomMeshComponent.cpp
--- a/Engine/Plugins/Runtime/CustomMeshComponent/Source/CustomMeshComponent/Private/CustomMeshComponent.cpp
+++ b/Engine/Plugins/Runtime/CustomMeshComponent/Source/CustomMeshComponent/Private/CustomMeshComponent.cpp
@@ -261,8 +261,8 @@ private:
 
 //////////////////////////////////////////////////////////////////////////
 
-UCustomMeshComponent::UCustomMeshComponent( const FPostConstructInitializeProperties& PCIP )
-	: Super( PCIP )
+UCustomMeshComponent::UCustomMeshComponent( const FObjectInitializer& ObjectInitializer )
+	: Super( ObjectInitializer )
 {
 	PrimaryComponentTick.bCanEverTick = false;
 
diff --git a/Engine/Plugins/Runtime/SoundMod/Source/SoundMod/Private/SoundMod.cpp b/Engine/Plugins/Runtime/SoundMod/Source/SoundMod/Private/SoundMod.cpp
--- a/Engine/Plugins/Runtime/SoundMod/Source/SoundMod/Private/SoundMod.cpp
+++ b/Engine/Plugins/Runtime/SoundMod/Source/SoundMod/Private/SoundMod.cpp
@@ -3,8 +3,8 @@
 #include "SoundModPrivatePCH.h"
 #include "SoundDefinitions.h"
 
-USoundMod::USoundMod(const class FPostConstructInitializeProperties& PCIP)
-	: Super(PCIP)
+USoundMod::USoundMod(const FObjectInitializer& ObjectInitializer)
+	: Super(ObjectInitializer)
 {
 }
 
diff --git a/Engine/Plugins/Runtime/SoundMod/Source/SoundMod/Private/SoundModWave.cpp b/Engine/Plugins/Runtime/SoundMod/Source/SoundMod/Private/SoundModWave.cpp
--- a/Engine/Plugins/Runtime/SoundMod/Source/SoundMod/Private/SoundModWave.cpp
+++ b/Engine/Plugins/Runtime/SoundMod/Source/SoundMod/Private/SoundModWave.cpp
@@ -2,8 +2,8 @@
 
 #include "SoundModPrivatePCH.h"
 
-USoundModWave::USoundModWave(const class FPostConstructInitializeProperties& PCIP)
-	: Super(PCIP)
+USoundModWave::USoundModWave(const FObjectInitializer& ObjectInitializer)
+	: Super(ObjectInitializer)
 	, xmpContext(NULL)
 {
 }
diff --git a/Engine/Plugins/Runtime/SoundMod/Source/SoundMod/Private/SoundNodeModPlayer.cpp b/Engine/Plugins/Runtime/SoundMod/Source/SoundMod/Private/SoundNodeModPlayer.cpp
--- a/Engine/Plugins/Runtime/SoundMod/Source/SoundMod/Private/SoundNodeModPlayer.cpp
+++ b/Engine/Plugins/Runtime/SoundMod/Source/SoundMod/Private/SoundNodeModPlayer.cpp
@@ -6,8 +6,8 @@
 
 #define LOCTEXT_NAMESPACE "SoundNodeModPlayer"
 
-USoundNodeModPlayer::USoundNodeModPlayer(const class FPostConstructInitializeProperties& PCIP)
-	: Super(PCIP)
+USoundNodeModPlayer::USoundNodeModPlayer(const FObjectInitializer& ObjectInitializer)
+	: Super(ObjectInitializer)
 {
 }
 
diff --git a/Engine/Plugins/Runtime/SoundMod/Source/SoundModImporter/Private/SoundModImporterFactory.cpp b/Engine/Plugins/Runtime/SoundMod/Source/SoundModImporter/Private/SoundModImporterFactory.cpp
--- a/Engine/Plugins/Runtime/SoundMod/Source/SoundModImporter/Private/SoundModImporterFactory.cpp
+++ b/Engine/Plugins/Runtime/SoundMod/Source/SoundModImporter/Private/SoundModImporterFactory.cpp
@@ -12,8 +12,8 @@
 
 static bool bSoundModFactorySuppressImportOverwriteDialog = false;
 
-USoundModImporterFactory::USoundModImporterFactory(const FPostConstructInitializeProperties& PCIP)
-	: Super(PCIP)
+USoundModImporterFactory::USoundModImporterFactory(const FObjectInitializer& ObjectInitializer)
+	: Super(ObjectInitializer)
 {
 	bCreateNew = false;
 	SupportedClass = USoundMod::StaticClass();
@@ -106,7 +106,7 @@ UObject* USoundModImporterFactory::FactoryCreateBinary
 
 	// Use pre-existing sound if it exists and we want to keep settings,
 	// otherwise create new sound and import raw data.
-	USoundMod* Sound = (bUseExistingSettings && ExistingSound) ? ExistingSound : new(InParent, Name, Flags) USoundMod(FPostConstructInitializeProperties());
+	USoundMod* Sound = (bUseExistingSettings && ExistingSound) ? ExistingSound : new(InParent, Name, Flags) USoundMod(FObjectInitializer());
 
 	Sound->Duration = xmpModuleInfo.seq_data->duration / 1000.f;
 
diff --git a/Engine/Plugins/Runtime/TwitchLiveStreaming/Source/TwitchLiveStreaming/TwitchProjectSettings.cpp b/Engine/Plugins/Runtime/TwitchLiveStreaming/Source/TwitchLiveStreaming/TwitchProjectSettings.cpp
--- a/Engine/Plugins/Runtime/TwitchLiveStreaming/Source/TwitchLiveStreaming/TwitchProjectSettings.cpp
+++ b/Engine/Plugins/Runtime/TwitchLiveStreaming/Source/TwitchLiveStreaming/TwitchProjectSettings.cpp
@@ -4,8 +4,8 @@
 #include "TwitchProjectSettings.h"
 
 
-UTwitchProjectSettings::UTwitchProjectSettings( const class FPostConstructInitializeProperties& PCIP )
-	: Super( PCIP )
+UTwitchProjectSettings::UTwitchProjectSettings( const FObjectInitializer& ObjectInitializer )
+	: Super( ObjectInitializer )
 {
 }
 
diff --git a/Engine/Plugins/ScriptPlugin/Source/ScriptEditorPlugin/Private/ScriptFactory.cpp b/Engine/Plugins/ScriptPlugin/Source/ScriptEditorPlugin/Private/ScriptFactory.cpp
--- a/Engine/Plugins/ScriptPlugin/Source/ScriptEditorPlugin/Private/ScriptFactory.cpp
+++ b/Engine/Plugins/ScriptPlugin/Source/ScriptEditorPlugin/Private/ScriptFactory.cpp
@@ -1,8 +1,8 @@
 // Copyright 1998-2014 Epic Games, Inc. All Rights Reserved. 
 #include "ScriptEditorPluginPrivatePCH.h"
 
-UScriptFactory::UScriptFactory(const FPostConstructInitializeProperties& PCIP)
-	: Super( PCIP )
+UScriptFactory::UScriptFactory(const FObjectInitializer& ObjectInitializer)
+	: Super( ObjectInitializer )
 {
 	SupportedClass = UScriptBlueprint::StaticClass();
 	ParentClass = AActor::StaticClass();
@@ -69,8 +69,8 @@ UObject* UScriptFactory::FactoryCreateText(UClass* InClass, UObject* InParent, F
 
 /** UReimportScriptFactory */
 
-UReimportScriptFactory::UReimportScriptFactory(const class FPostConstructInitializeProperties& PCIP)
-: Super(PCIP)
+UReimportScriptFactory::UReimportScriptFactory(const class FObjectInitializer& ObjectInitializer)
+: Super(ObjectInitializer)
 {
 }
 
diff --git a/Engine/Plugins/ScriptPlugin/Source/ScriptPlugin/Private/ScriptBlueprint.cpp b/Engine/Plugins/ScriptPlugin/Source/ScriptPlugin/Private/ScriptBlueprint.cpp
--- a/Engine/Plugins/ScriptPlugin/Source/ScriptPlugin/Private/ScriptBlueprint.cpp
+++ b/Engine/Plugins/ScriptPlugin/Source/ScriptPlugin/Private/ScriptBlueprint.cpp
@@ -8,8 +8,8 @@
 /////////////////////////////////////////////////////
 // UScriptBlueprint
 
-UScriptBlueprint::UScriptBlueprint(const FPostConstructInitializeProperties& PCIP)
-	: Super(PCIP)
+UScriptBlueprint::UScriptBlueprint(const FObjectInitializer& ObjectInitializer)
+	: Super(ObjectInitializer)
 {
 }
 
diff --git a/Engine/Plugins/ScriptPlugin/Source/ScriptPlugin/Private/ScriptBlueprintGeneratedClass.cpp b/Engine/Plugins/ScriptPlugin/Source/ScriptPlugin/Private/ScriptBlueprintGeneratedClass.cpp
--- a/Engine/Plugins/ScriptPlugin/Source/ScriptPlugin/Private/ScriptBlueprintGeneratedClass.cpp
+++ b/Engine/Plugins/ScriptPlugin/Source/ScriptPlugin/Private/ScriptBlueprintGeneratedClass.cpp
@@ -8,8 +8,8 @@
 /////////////////////////////////////////////////////
 // UScriptBlueprintGeneratedClass
 
-UScriptBlueprintGeneratedClass::UScriptBlueprintGeneratedClass(const FPostConstructInitializeProperties& PCIP)
-	: Super(PCIP)
+UScriptBlueprintGeneratedClass::UScriptBlueprintGeneratedClass(const FObjectInitializer& ObjectInitializer)
+	: Super(ObjectInitializer)
 {
 }
 
diff --git a/Engine/Plugins/ScriptPlugin/Source/ScriptPlugin/Private/ScriptComponent.cpp b/Engine/Plugins/ScriptPlugin/Source/ScriptPlugin/Private/ScriptComponent.cpp
--- a/Engine/Plugins/ScriptPlugin/Source/ScriptPlugin/Private/ScriptComponent.cpp
+++ b/Engine/Plugins/ScriptPlugin/Source/ScriptPlugin/Private/ScriptComponent.cpp
@@ -3,8 +3,8 @@
 
 //////////////////////////////////////////////////////////////////////////
 
-UScriptComponent::UScriptComponent(const FPostConstructInitializeProperties& PCIP)
-	: Super(PCIP)
+UScriptComponent::UScriptComponent(const FObjectInitializer& ObjectInitializer)
+	: Super(ObjectInitializer)
 {
 	PrimaryComponentTick.bCanEverTick = true;
 	bTickInEditor = false;
diff --git a/Engine/Plugins/ScriptPlugin/Source/ScriptPlugin/Private/ScriptContext.cpp b/Engine/Plugins/ScriptPlugin/Source/ScriptPlugin/Private/ScriptContext.cpp
--- a/Engine/Plugins/ScriptPlugin/Source/ScriptPlugin/Private/ScriptContext.cpp
+++ b/Engine/Plugins/ScriptPlugin/Source/ScriptPlugin/Private/ScriptContext.cpp
@@ -3,8 +3,8 @@
 
 //////////////////////////////////////////////////////////////////////////
 
-UScriptContext::UScriptContext(const FPostConstructInitializeProperties& PCIP)
-	: Super(PCIP)
+UScriptContext::UScriptContext(const FObjectInitializer& ObjectInitializer)
+	: Super(ObjectInitializer)
 {
 	Context = NULL;
 }
diff --git a/Engine/Plugins/ScriptPlugin/Source/ScriptPlugin/Private/ScriptContextComponent.cpp b/Engine/Plugins/ScriptPlugin/Source/ScriptPlugin/Private/ScriptContextComponent.cpp
--- a/Engine/Plugins/ScriptPlugin/Source/ScriptPlugin/Private/ScriptContextComponent.cpp
+++ b/Engine/Plugins/ScriptPlugin/Source/ScriptPlugin/Private/ScriptContextComponent.cpp
@@ -3,8 +3,8 @@
 
 //////////////////////////////////////////////////////////////////////////
 
-UScriptContextComponent::UScriptContextComponent(const FPostConstructInitializeProperties& PCIP)
-	: Super(PCIP)
+UScriptContextComponent::UScriptContextComponent(const FObjectInitializer& ObjectInitializer)
+	: Super(ObjectInitializer)
 {
 	PrimaryComponentTick.bCanEverTick = true;
 	bTickInEditor = false;
diff --git a/Engine/Plugins/ScriptPlugin/Source/ScriptPlugin/Private/ScriptTestActor.cpp b/Engine/Plugins/ScriptPlugin/Source/ScriptPlugin/Private/ScriptTestActor.cpp
--- a/Engine/Plugins/ScriptPlugin/Source/ScriptPlugin/Private/ScriptTestActor.cpp
+++ b/Engine/Plugins/ScriptPlugin/Source/ScriptPlugin/Private/ScriptTestActor.cpp
@@ -3,8 +3,8 @@
 
 //////////////////////////////////////////////////////////////////////////
 
-AScriptTestActor::AScriptTestActor(const FPostConstructInitializeProperties& PCIP)
-	: Super( PCIP )
+AScriptTestActor::AScriptTestActor(const FObjectInitializer& ObjectInitializer)
+	: Super( ObjectInitializer )
 {
 }
 
diff --git a/Engine/Plugins/Slate/SlateRemote/Source/SlateRemote/Private/SlateRemoteClasses.cpp b/Engine/Plugins/Slate/SlateRemote/Source/SlateRemote/Private/SlateRemoteClasses.cpp
--- a/Engine/Plugins/Slate/SlateRemote/Source/SlateRemote/Private/SlateRemoteClasses.cpp
+++ b/Engine/Plugins/Slate/SlateRemote/Source/SlateRemote/Private/SlateRemoteClasses.cpp
@@ -3,8 +3,8 @@
 #include "SlateRemotePrivatePCH.h"
 
 
-USlateRemoteSettings::USlateRemoteSettings( const class FPostConstructInitializeProperties& PCIP )
-	: Super(PCIP)
+USlateRemoteSettings::USlateRemoteSettings( const FObjectInitializer& ObjectInitializer )
+	: Super(ObjectInitializer)
 	, EnableRemoteServer(true)
 	, EditorServerEndpoint(SLATE_REMOTE_SERVER_DEFAULT_EDITOR_ENDPOINT.ToString())
 	, GameServerEndpoint(SLATE_REMOTE_SERVER_DEFAULT_GAME_ENDPOINT.ToString())
diff --git a/Engine/Source/Developer/AITestSuite/Private/Actions/TestPawnAction_CallFunction.cpp b/Engine/Source/Developer/AITestSuite/Private/Actions/TestPawnAction_CallFunction.cpp
--- a/Engine/Source/Developer/AITestSuite/Private/Actions/TestPawnAction_CallFunction.cpp
+++ b/Engine/Source/Developer/AITestSuite/Private/Actions/TestPawnAction_CallFunction.cpp
@@ -3,8 +3,8 @@
 #include "AITestSuitePrivatePCH.h"
 #include "Actions/TestPawnAction_CallFunction.h"
 
-UTestPawnAction_CallFunction::UTestPawnAction_CallFunction(const FPostConstructInitializeProperties& PCIP)
-	: Super(PCIP)
+UTestPawnAction_CallFunction::UTestPawnAction_CallFunction(const FObjectInitializer& ObjectInitializer)
+	: Super(ObjectInitializer)
 	, FunctionToCall(nullptr)
 {
 
diff --git a/Engine/Source/Developer/AITestSuite/Private/Actions/TestPawnAction_Log.cpp b/Engine/Source/Developer/AITestSuite/Private/Actions/TestPawnAction_Log.cpp
--- a/Engine/Source/Developer/AITestSuite/Private/Actions/TestPawnAction_Log.cpp
+++ b/Engine/Source/Developer/AITestSuite/Private/Actions/TestPawnAction_Log.cpp
@@ -2,8 +2,8 @@
 
 #include "AITestSuitePrivatePCH.h"
 
-UTestPawnAction_Log::UTestPawnAction_Log(const FPostConstructInitializeProperties& PCIP) 
-	: Super(PCIP)
+UTestPawnAction_Log::UTestPawnAction_Log(const FObjectInitializer& ObjectInitializer) 
+	: Super(ObjectInitializer)
 	, Logger(nullptr)
 {
 
diff --git a/Engine/Source/Developer/AITestSuite/Private/BehaviorTree/TestBTDecorator_CantExecute.cpp b/Engine/Source/Developer/AITestSuite/Private/BehaviorTree/TestBTDecorator_CantExecute.cpp
--- a/Engine/Source/Developer/AITestSuite/Private/BehaviorTree/TestBTDecorator_CantExecute.cpp
+++ b/Engine/Source/Developer/AITestSuite/Private/BehaviorTree/TestBTDecorator_CantExecute.cpp
@@ -3,7 +3,7 @@
 #include "AITestSuitePrivatePCH.h"
 #include "BehaviorTree/TestBTDecorator_CantExecute.h"
 
-UTestBTDecorator_CantExecute::UTestBTDecorator_CantExecute(const FPostConstructInitializeProperties& PCIP) : Super(PCIP)
+UTestBTDecorator_CantExecute::UTestBTDecorator_CantExecute(const FObjectInitializer& ObjectInitializer) : Super(ObjectInitializer)
 {
 	NodeName = TEXT("Can't Exexcute");
 
diff --git a/Engine/Source/Developer/AITestSuite/Private/BehaviorTree/TestBTDecorator_DelayedAbort.cpp b/Engine/Source/Developer/AITestSuite/Private/BehaviorTree/TestBTDecorator_DelayedAbort.cpp
--- a/Engine/Source/Developer/AITestSuite/Private/BehaviorTree/TestBTDecorator_DelayedAbort.cpp
+++ b/Engine/Source/Developer/AITestSuite/Private/BehaviorTree/TestBTDecorator_DelayedAbort.cpp
@@ -3,7 +3,7 @@
 #include "AITestSuitePrivatePCH.h"
 #include "BehaviorTree/TestBTDecorator_DelayedAbort.h"
 
-UTestBTDecorator_DelayedAbort::UTestBTDecorator_DelayedAbort(const FPostConstructInitializeProperties& PCIP) : Super(PCIP)
+UTestBTDecorator_DelayedAbort::UTestBTDecorator_DelayedAbort(const FObjectInitializer& ObjectInitializer) : Super(ObjectInitializer)
 {
 	NodeName = "Delayed Abort";
 	DelayTicks = 5;
diff --git a/Engine/Source/Developer/AITestSuite/Private/BehaviorTree/TestBTTask_Log.cpp b/Engine/Source/Developer/AITestSuite/Private/BehaviorTree/TestBTTask_Log.cpp
--- a/Engine/Source/Developer/AITestSuite/Private/BehaviorTree/TestBTTask_Log.cpp
+++ b/Engine/Source/Developer/AITestSuite/Private/BehaviorTree/TestBTTask_Log.cpp
@@ -4,7 +4,7 @@
 #include "BehaviorTree/TestBTTask_Log.h"
 #include "MockAI_BT.h"
 
-UTestBTTask_Log::UTestBTTask_Log(const FPostConstructInitializeProperties& PCIP) : Super(PCIP)
+UTestBTTask_Log::UTestBTTask_Log(const FObjectInitializer& ObjectInitializer) : Super(ObjectInitializer)
 {
 	NodeName = "Log";
 	ExecutionTicks = 0;
diff --git a/Engine/Source/Developer/AITestSuite/Private/BehaviorTree/TestBTTask_SetFlag.cpp b/Engine/Source/Developer/AITestSuite/Private/BehaviorTree/TestBTTask_SetFlag.cpp
--- a/Engine/Source/Developer/AITestSuite/Private/BehaviorTree/TestBTTask_SetFlag.cpp
+++ b/Engine/Source/Developer/AITestSuite/Private/BehaviorTree/TestBTTask_SetFlag.cpp
@@ -4,7 +4,7 @@
 #include "BehaviorTree/TestBTTask_SetFlag.h"
 #include "BehaviorTree/BlackboardComponent.h"
 
-UTestBTTask_SetFlag::UTestBTTask_SetFlag(const FPostConstructInitializeProperties& PCIP) : Super(PCIP)
+UTestBTTask_SetFlag::UTestBTTask_SetFlag(const FObjectInitializer& ObjectInitializer) : Super(ObjectInitializer)
 {
 	NodeName = "Log";
 	TaskResult = EBTNodeResult::Succeeded;
diff --git a/Engine/Source/Developer/AITestSuite/Private/MockAI/MockAI.cpp b/Engine/Source/Developer/AITestSuite/Private/MockAI/MockAI.cpp
--- a/Engine/Source/Developer/AITestSuite/Private/MockAI/MockAI.cpp
+++ b/Engine/Source/Developer/AITestSuite/Private/MockAI/MockAI.cpp
@@ -22,8 +22,8 @@ TStatId FTestTickHelper::GetStatId() const
 //----------------------------------------------------------------------//
 // 
 //----------------------------------------------------------------------//
-UMockAI::UMockAI(const FPostConstructInitializeProperties& PCIP)
-	: Super(PCIP)
+UMockAI::UMockAI(const FObjectInitializer& ObjectInitializer)
+	: Super(ObjectInitializer)
 {
 }
 
diff --git a/Engine/Source/Developer/AITestSuite/Private/MockAI/MockAI_BT.cpp b/Engine/Source/Developer/AITestSuite/Private/MockAI/MockAI_BT.cpp
--- a/Engine/Source/Developer/AITestSuite/Private/MockAI/MockAI_BT.cpp
+++ b/Engine/Source/Developer/AITestSuite/Private/MockAI/MockAI_BT.cpp
@@ -8,8 +8,8 @@
 //----------------------------------------------------------------------//
 TArray<int32> UMockAI_BT::ExecutionLog;
 
-UMockAI_BT::UMockAI_BT(const FPostConstructInitializeProperties& PCIP)
-	: Super(PCIP)
+UMockAI_BT::UMockAI_BT(const FObjectInitializer& ObjectInitializer)
+	: Super(ObjectInitializer)
 {
 	if (HasAnyFlags(RF_ClassDefaultObject) == false)
 	{
diff --git a/Engine/Source/Developer/FunctionalTesting/Private/FuncTestRenderingComponent.cpp b/Engine/Source/Developer/FunctionalTesting/Private/FuncTestRenderingComponent.cpp
--- a/Engine/Source/Developer/FunctionalTesting/Private/FuncTestRenderingComponent.cpp
+++ b/Engine/Source/Developer/FunctionalTesting/Private/FuncTestRenderingComponent.cpp
@@ -38,8 +38,8 @@ protected:
 //----------------------------------------------------------------------//
 // UFuncTestRenderingComponent
 //----------------------------------------------------------------------//
-UFuncTestRenderingComponent::UFuncTestRenderingComponent(const class FPostConstructInitializeProperties& PCIP)
-	: Super(PCIP)
+UFuncTestRenderingComponent::UFuncTestRenderingComponent(const FObjectInitializer& ObjectInitializer)
+	: Super(ObjectInitializer)
 {
 	SetCollisionProfileName(UCollisionProfile::NoCollision_ProfileName);
 
diff --git a/Engine/Source/Developer/FunctionalTesting/Private/FunctionalAITest.cpp b/Engine/Source/Developer/FunctionalTesting/Private/FunctionalAITest.cpp
--- a/Engine/Source/Developer/FunctionalTesting/Private/FunctionalAITest.cpp
+++ b/Engine/Source/Developer/FunctionalTesting/Private/FunctionalAITest.cpp
@@ -6,8 +6,8 @@
 #include "BehaviorTree/BehaviorTree.h"
 #include "AIController.h"
 
-AFunctionalAITest::AFunctionalAITest( const class FPostConstructInitializeProperties& PCIP )
-	: Super(PCIP)
+AFunctionalAITest::AFunctionalAITest( const FObjectInitializer& ObjectInitializer )
+	: Super(ObjectInitializer)
 	, CurrentSpawnSetIndex(INDEX_NONE)
 	, bSingleSetRun(false)
 {
diff --git a/Engine/Source/Developer/FunctionalTesting/Private/FunctionalTest.cpp b/Engine/Source/Developer/FunctionalTesting/Private/FunctionalTest.cpp
--- a/Engine/Source/Developer/FunctionalTesting/Private/FunctionalTest.cpp
+++ b/Engine/Source/Developer/FunctionalTesting/Private/FunctionalTest.cpp
@@ -4,8 +4,8 @@
 #include "ObjectEditorUtils.h"
 #include "VisualLog.h"
 
-AFunctionalTest::AFunctionalTest( const class FPostConstructInitializeProperties& PCIP )
-	: Super(PCIP)
+AFunctionalTest::AFunctionalTest( const FObjectInitializer& ObjectInitializer )
+	: Super(ObjectInitializer)
 	, TimesUpResult(EFunctionalTestResult::Failed)
 	, TimeLimit(DefaultTimeLimit)
 	, TimesUpMessage( NSLOCTEXT("FunctionalTest", "DefaultTimesUpMessage", "Time's up!") )
@@ -16,7 +16,7 @@ AFunctionalTest::AFunctionalTest( const class FPostConstructInitializeProperties
 	PrimaryActorTick.bCanEverTick = true;
 	PrimaryActorTick.bStartWithTickEnabled = false;
 	
-	SpriteComponent = PCIP.CreateDefaultSubobject<UBillboardComponent>(this, TEXT("Sprite"));
+	SpriteComponent = ObjectInitializer.CreateDefaultSubobject<UBillboardComponent>(this, TEXT("Sprite"));
 	if (SpriteComponent)
 	{
 		SpriteComponent->bHiddenInGame = false;
diff --git a/Engine/Source/Developer/FunctionalTesting/Private/FunctionalTestingManager.cpp b/Engine/Source/Developer/FunctionalTesting/Private/FunctionalTestingManager.cpp
--- a/Engine/Source/Developer/FunctionalTesting/Private/FunctionalTestingManager.cpp
+++ b/Engine/Source/Developer/FunctionalTesting/Private/FunctionalTestingManager.cpp
@@ -55,8 +55,8 @@ namespace FFunctionalTesting
 
 FAutomationTestExecutionInfo* UFunctionalTestingManager::ExecutionInfo = NULL;
 
-UFunctionalTestingManager::UFunctionalTestingManager( const class FPostConstructInitializeProperties& PCIP )
-	: Super(PCIP)
+UFunctionalTestingManager::UFunctionalTestingManager( const FObjectInitializer& ObjectInitializer )
+	: Super(ObjectInitializer)
 	, bIsRunning(false)
 	, bLooped(false)
 	, bWaitForNavigationBuildFinish(false)
diff --git a/Engine/Source/Developer/GameplayDebugger/Private/GameplayDebuggingComponent.cpp b/Engine/Source/Developer/GameplayDebugger/Private/GameplayDebuggingComponent.cpp
--- a/Engine/Source/Developer/GameplayDebugger/Private/GameplayDebuggingComponent.cpp
+++ b/Engine/Source/Developer/GameplayDebugger/Private/GameplayDebuggingComponent.cpp
@@ -123,8 +123,8 @@ protected:
 FName UGameplayDebuggingComponent::DefaultComponentName = TEXT("GameplayDebuggingComponent");
 FOnDebuggingTargetChanged UGameplayDebuggingComponent::OnDebuggingTargetChangedDelegate;
 
-UGameplayDebuggingComponent::UGameplayDebuggingComponent(const class FPostConstructInitializeProperties& PCIP) 
-	: Super(PCIP)
+UGameplayDebuggingComponent::UGameplayDebuggingComponent(const FObjectInitializer& ObjectInitializer) 
+	: Super(ObjectInitializer)
 	, bDrawEQSLabels(true)
 	, bDrawEQSFailedItems(true)
 {
diff --git a/Engine/Source/Developer/GameplayDebugger/Private/GameplayDebuggingControllerComponent.cpp b/Engine/Source/Developer/GameplayDebugger/Private/GameplayDebuggingControllerComponent.cpp
--- a/Engine/Source/Developer/GameplayDebugger/Private/GameplayDebuggingControllerComponent.cpp
+++ b/Engine/Source/Developer/GameplayDebugger/Private/GameplayDebuggingControllerComponent.cpp
@@ -13,8 +13,8 @@ DEFINE_LOG_CATEGORY_STATIC(LogGameplayDebugging, Log, All);
 #define BUGIT_VIEWS (1<<EAIDebugDrawDataView::Basic) | (1 << EAIDebugDrawDataView::OverHead)
 #define BREAK_LINE_TEXT TEXT("________________________________________________________________")
 
-UGameplayDebuggingControllerComponent::UGameplayDebuggingControllerComponent(const class FPostConstructInitializeProperties& PCIP)
-	: Super(PCIP)
+UGameplayDebuggingControllerComponent::UGameplayDebuggingControllerComponent(const FObjectInitializer& ObjectInitializer)
+	: Super(ObjectInitializer)
 	, KeyPressActivationTime(0.4f)
 {
 	PrimaryComponentTick.bCanEverTick = true;
diff --git a/Engine/Source/Developer/GameplayDebugger/Private/GameplayDebuggingHUDComponent.cpp b/Engine/Source/Developer/GameplayDebugger/Private/GameplayDebuggingHUDComponent.cpp
--- a/Engine/Source/Developer/GameplayDebugger/Private/GameplayDebuggingHUDComponent.cpp
+++ b/Engine/Source/Developer/GameplayDebugger/Private/GameplayDebuggingHUDComponent.cpp
@@ -28,8 +28,8 @@
 DEFINE_LOG_CATEGORY_STATIC(LogHUD, Log, All);
 
 
-AGameplayDebuggingHUDComponent::AGameplayDebuggingHUDComponent(const class FPostConstructInitializeProperties& PCIP)
-	: Super(PCIP)
+AGameplayDebuggingHUDComponent::AGameplayDebuggingHUDComponent(const FObjectInitializer& ObjectInitializer)
+	: Super(ObjectInitializer)
 	, EngineShowFlags(EShowFlagInitMode::ESFIM_Game)
 {
 	World = NULL;
diff --git a/Engine/Source/Developer/GameplayDebugger/Private/GameplayDebuggingReplicator.cpp b/Engine/Source/Developer/GameplayDebugger/Private/GameplayDebuggingReplicator.cpp
--- a/Engine/Source/Developer/GameplayDebugger/Private/GameplayDebuggingReplicator.cpp
+++ b/Engine/Source/Developer/GameplayDebugger/Private/GameplayDebuggingReplicator.cpp
@@ -16,8 +16,8 @@
 
 FOnSelectionChanged AGameplayDebuggingReplicator::OnSelectionChangedDelegate;
 
-AGameplayDebuggingReplicator::AGameplayDebuggingReplicator(const class FPostConstructInitializeProperties& PCIP)
-	: Super(PCIP)
+AGameplayDebuggingReplicator::AGameplayDebuggingReplicator(const FObjectInitializer& ObjectInitializer)
+	: Super(ObjectInitializer)
 	, MaxEQSQueries(5)
 	, bIsGlobalInWorld(true)
 	, LastDrawAtFrame(0)
@@ -41,7 +41,7 @@ AGameplayDebuggingReplicator::AGameplayDebuggingReplicator(const class FPostCons
 	PrimaryActorTick.bCanEverTick = true;
 	PrimaryActorTick.bStartWithTickEnabled = false;
 	
-	TSubobjectPtr<USceneComponent> SceneComponent = PCIP.CreateDefaultSubobject<USceneComponent>(this, TEXT("SceneComponent"));
+	TSubobjectPtr<USceneComponent> SceneComponent = ObjectInitializer.CreateDefaultSubobject<USceneComponent>(this, TEXT("SceneComponent"));
 	RootComponent = SceneComponent;
 
 #if WITH_EDITOR
diff --git a/Engine/Source/Developer/HTML5/HTML5TargetPlatform/Private/HTML5TargetPlatformClasses.cpp b/Engine/Source/Developer/HTML5/HTML5TargetPlatform/Private/HTML5TargetPlatformClasses.cpp
--- a/Engine/Source/Developer/HTML5/HTML5TargetPlatform/Private/HTML5TargetPlatformClasses.cpp
+++ b/Engine/Source/Developer/HTML5/HTML5TargetPlatform/Private/HTML5TargetPlatformClasses.cpp
@@ -11,6 +11,6 @@
 /* UHTML5TargetSettings structors
  *****************************************************************************/
 
-UHTML5TargetSettings::UHTML5TargetSettings( const class FPostConstructInitializeProperties& PCIP )
-	: Super(PCIP)
+UHTML5TargetSettings::UHTML5TargetSettings( const FObjectInitializer& ObjectInitializer )
+	: Super(ObjectInitializer)
 { }
diff --git a/Engine/Source/Developer/HotReload/Private/HotReloadClassReinstancer.cpp b/Engine/Source/Developer/HotReload/Private/HotReloadClassReinstancer.cpp
--- a/Engine/Source/Developer/HotReload/Private/HotReloadClassReinstancer.cpp
+++ b/Engine/Source/Developer/HotReload/Private/HotReloadClassReinstancer.cpp
@@ -165,7 +165,7 @@ void FHotReloadClassReinstancer::ReconstructClassDefaultObject(UClass* InOldClas
 	new ((void *)OldCDO) UObjectBase(InOldClass, CDOFlags, CDOOuter, CDOName);
 	const bool bShouldInitilizeProperties = false;
 	const bool bCopyTransientsFromClassDefaults = false;
-	(*InOldClass->ClassConstructor)(FPostConstructInitializeProperties(OldCDO, ParentDefaultObject, bCopyTransientsFromClassDefaults, bShouldInitilizeProperties));
+	(*InOldClass->ClassConstructor)(FObjectInitializer(OldCDO, ParentDefaultObject, bCopyTransientsFromClassDefaults, bShouldInitilizeProperties));
 }
 
 void FHotReloadClassReinstancer::RecreateCDOAndSetupOldClassReinstancing(UClass* InOldClass)
diff --git a/Engine/Source/Developer/InternationalizationSettings/Private/InternationalizationSettingsModel.cpp b/Engine/Source/Developer/InternationalizationSettings/Private/InternationalizationSettingsModel.cpp
--- a/Engine/Source/Developer/InternationalizationSettings/Private/InternationalizationSettingsModel.cpp
+++ b/Engine/Source/Developer/InternationalizationSettings/Private/InternationalizationSettingsModel.cpp
@@ -6,8 +6,8 @@
 /* UInternationalizationSettingsModel interface
  *****************************************************************************/
 
-UInternationalizationSettingsModel::UInternationalizationSettingsModel( const class FPostConstructInitializeProperties& PCIP )
-	: Super(PCIP)
+UInternationalizationSettingsModel::UInternationalizationSettingsModel( const FObjectInitializer& ObjectInitializer )
+	: Super(ObjectInitializer)
 { }
 
 
diff --git a/Engine/Source/Developer/Linux/LinuxTargetPlatform/Private/LinuxTargetPlatformClasses.cpp b/Engine/Source/Developer/Linux/LinuxTargetPlatform/Private/LinuxTargetPlatformClasses.cpp
--- a/Engine/Source/Developer/Linux/LinuxTargetPlatform/Private/LinuxTargetPlatformClasses.cpp
+++ b/Engine/Source/Developer/Linux/LinuxTargetPlatform/Private/LinuxTargetPlatformClasses.cpp
@@ -11,6 +11,6 @@
 /* UAndroidTargetSettings structors
  *****************************************************************************/
 
-ULinuxTargetSettings::ULinuxTargetSettings( const class FPostConstructInitializeProperties& PCIP )
-	: Super(PCIP)
+ULinuxTargetSettings::ULinuxTargetSettings( const FObjectInitializer& ObjectInitializer )
+	: Super(ObjectInitializer)
 { }
diff --git a/Engine/Source/Developer/LogVisualizer/Private/LogVisualizerDebugActor.cpp b/Engine/Source/Developer/LogVisualizer/Private/LogVisualizerDebugActor.cpp
--- a/Engine/Source/Developer/LogVisualizer/Private/LogVisualizerDebugActor.cpp
+++ b/Engine/Source/Developer/LogVisualizer/Private/LogVisualizerDebugActor.cpp
@@ -9,7 +9,7 @@
 #endif // WITH_EDITOR
 #include "LogVisualizerDebugActor.h"
 
-ALogVisualizerDebugActor::ALogVisualizerDebugActor(const class FPostConstructInitializeProperties& PCIP)
-	: Super(PCIP)
+ALogVisualizerDebugActor::ALogVisualizerDebugActor(const FObjectInitializer& ObjectInitializer)
+	: Super(ObjectInitializer)
 {
 }
diff --git a/Engine/Source/Developer/Mac/MacTargetPlatform/Private/MacTargetPlatformClasses.cpp b/Engine/Source/Developer/Mac/MacTargetPlatform/Private/MacTargetPlatformClasses.cpp
--- a/Engine/Source/Developer/Mac/MacTargetPlatform/Private/MacTargetPlatformClasses.cpp
+++ b/Engine/Source/Developer/Mac/MacTargetPlatform/Private/MacTargetPlatformClasses.cpp
@@ -11,6 +11,6 @@
 /* UMacTargetSettings structors
  *****************************************************************************/
 
-UMacTargetSettings::UMacTargetSettings( const class FPostConstructInitializeProperties& PCIP )
-	: Super(PCIP)
+UMacTargetSettings::UMacTargetSettings( const FObjectInitializer& ObjectInitializer )
+	: Super(ObjectInitializer)
 { }
diff --git a/Engine/Source/Developer/MeshUtilities/Private/MeshUtilities.cpp b/Engine/Source/Developer/MeshUtilities/Private/MeshUtilities.cpp
--- a/Engine/Source/Developer/MeshUtilities/Private/MeshUtilities.cpp
+++ b/Engine/Source/Developer/MeshUtilities/Private/MeshUtilities.cpp
@@ -3220,7 +3220,7 @@ void FMeshUtilities::CreateProxyMesh(
 		MeshPackage->Modify();
 	}
 
-	UStaticMesh* StaticMesh = new(MeshPackage, FName(*(TEXT("SM_") + AssetBaseName)), RF_Public|RF_Standalone) UStaticMesh(FPostConstructInitializeProperties());
+	UStaticMesh* StaticMesh = new(MeshPackage, FName(*(TEXT("SM_") + AssetBaseName)), RF_Public|RF_Standalone) UStaticMesh(FObjectInitializer());
 	StaticMesh->InitResources();
 	{
 		FString OutputPath = StaticMesh->GetPathName();
@@ -3793,7 +3793,7 @@ void FMeshUtilities::MergeActors(
 		Package->FullyLoad();
 		Package->Modify();
 
-		UStaticMesh* StaticMesh = new(Package, *AssetName, RF_Public|RF_Standalone) UStaticMesh(FPostConstructInitializeProperties());
+		UStaticMesh* StaticMesh = new(Package, *AssetName, RF_Public|RF_Standalone) UStaticMesh(FObjectInitializer());
 		StaticMesh->InitResources();
 		
 		FString OutputPath = StaticMesh->GetPathName();
diff --git a/Engine/Source/Developer/MovieSceneTools/Private/K2Node_PlayMovieScene.cpp b/Engine/Source/Developer/MovieSceneTools/Private/K2Node_PlayMovieScene.cpp
--- a/Engine/Source/Developer/MovieSceneTools/Private/K2Node_PlayMovieScene.cpp
+++ b/Engine/Source/Developer/MovieSceneTools/Private/K2Node_PlayMovieScene.cpp
@@ -26,8 +26,8 @@ UEdGraphNode* FEdGraphSchemaAction_K2AddPlayMovieScene::PerformAction(class UEdG
 }
 
 
-UK2Node_PlayMovieScene::UK2Node_PlayMovieScene(const class FPostConstructInitializeProperties& PCIP)
-	: Super(PCIP)
+UK2Node_PlayMovieScene::UK2Node_PlayMovieScene(const FObjectInitializer& ObjectInitializer)
+	: Super(ObjectInitializer)
 {
 }
 
diff --git a/Engine/Source/Developer/MovieSceneTools/Private/MovieSceneFactory.cpp b/Engine/Source/Developer/MovieSceneTools/Private/MovieSceneFactory.cpp
--- a/Engine/Source/Developer/MovieSceneTools/Private/MovieSceneFactory.cpp
+++ b/Engine/Source/Developer/MovieSceneTools/Private/MovieSceneFactory.cpp
@@ -6,8 +6,8 @@
 
 #define LOCTEXT_NAMESPACE "MovieSceneFactory"
 
-UMovieSceneFactory::UMovieSceneFactory( const class FPostConstructInitializeProperties& PCIP )
-	: Super(PCIP)
+UMovieSceneFactory::UMovieSceneFactory( const FObjectInitializer& ObjectInitializer )
+	: Super(ObjectInitializer)
 {
 
 	// @todo sequencer: Only allow users to create new MovieScenes if that feature is turned on globally.
diff --git a/Engine/Source/Developer/SourceCodeAccess/Private/SourceCodeAccessSettings.cpp b/Engine/Source/Developer/SourceCodeAccess/Private/SourceCodeAccessSettings.cpp
--- a/Engine/Source/Developer/SourceCodeAccess/Private/SourceCodeAccessSettings.cpp
+++ b/Engine/Source/Developer/SourceCodeAccess/Private/SourceCodeAccessSettings.cpp
@@ -3,8 +3,8 @@
 #include "SourceCodeAccessPrivatePCH.h"
 #include "SourceCodeAccessSettings.h"
 
-USourceCodeAccessSettings::USourceCodeAccessSettings(const class FPostConstructInitializeProperties& PCIP)
-	: Super(PCIP)
+USourceCodeAccessSettings::USourceCodeAccessSettings(const FObjectInitializer& ObjectInitializer)
+	: Super(ObjectInitializer)
 {
 #if PLATFORM_WINDOWS
 	PreferredAccessor = TEXT("VisualStudioSourceCodeAccessor");
diff --git a/Engine/Source/Developer/Windows/WindowsTargetPlatform/Private/WindowsTargetPlatformClasses.cpp b/Engine/Source/Developer/Windows/WindowsTargetPlatform/Private/WindowsTargetPlatformClasses.cpp
--- a/Engine/Source/Developer/Windows/WindowsTargetPlatform/Private/WindowsTargetPlatformClasses.cpp
+++ b/Engine/Source/Developer/Windows/WindowsTargetPlatform/Private/WindowsTargetPlatformClasses.cpp
@@ -6,6 +6,6 @@
 /* UWindowsTargetSettings structors
  *****************************************************************************/
 
-UWindowsTargetSettings::UWindowsTargetSettings( const class FPostConstructInitializeProperties& PCIP )
-	: Super(PCIP)
+UWindowsTargetSettings::UWindowsTargetSettings( const FObjectInitializer& ObjectInitializer )
+	: Super(ObjectInitializer)
 { }
diff --git a/Engine/Source/Editor/AnimGraph/Private/AnimGraphNode_ApplyAdditive.cpp b/Engine/Source/Editor/AnimGraph/Private/AnimGraphNode_ApplyAdditive.cpp
--- a/Engine/Source/Editor/AnimGraph/Private/AnimGraphNode_ApplyAdditive.cpp
+++ b/Engine/Source/Editor/AnimGraph/Private/AnimGraphNode_ApplyAdditive.cpp
@@ -11,8 +11,8 @@
 
 #define LOCTEXT_NAMESPACE "A3Nodes"
 
-UAnimGraphNode_ApplyAdditive::UAnimGraphNode_ApplyAdditive(const FPostConstructInitializeProperties& PCIP)
-	: Super(PCIP)
+UAnimGraphNode_ApplyAdditive::UAnimGraphNode_ApplyAdditive(const FObjectInitializer& ObjectInitializer)
+	: Super(ObjectInitializer)
 {
 }
 
diff --git a/Engine/Source/Editor/AnimGraph/Private/AnimGraphNode_Base.cpp b/Engine/Source/Editor/AnimGraph/Private/AnimGraphNode_Base.cpp
--- a/Engine/Source/Editor/AnimGraph/Private/AnimGraphNode_Base.cpp
+++ b/Engine/Source/Editor/AnimGraph/Private/AnimGraphNode_Base.cpp
@@ -115,8 +115,8 @@ public:
 /////////////////////////////////////////////////////
 // UAnimGraphNode_Base
 
-UAnimGraphNode_Base::UAnimGraphNode_Base(const FPostConstructInitializeProperties& PCIP)
-	: Super(PCIP)
+UAnimGraphNode_Base::UAnimGraphNode_Base(const FObjectInitializer& ObjectInitializer)
+	: Super(ObjectInitializer)
 {
 }
 
diff --git a/Engine/Source/Editor/AnimGraph/Private/AnimGraphNode_BlendListBase.cpp b/Engine/Source/Editor/AnimGraph/Private/AnimGraphNode_BlendListBase.cpp
--- a/Engine/Source/Editor/AnimGraph/Private/AnimGraphNode_BlendListBase.cpp
+++ b/Engine/Source/Editor/AnimGraph/Private/AnimGraphNode_BlendListBase.cpp
@@ -6,8 +6,8 @@
 /////////////////////////////////////////////////////
 // UAnimGraphNode_BlendListBase
 
-UAnimGraphNode_BlendListBase::UAnimGraphNode_BlendListBase(const FPostConstructInitializeProperties& PCIP)
-	: Super(PCIP)
+UAnimGraphNode_BlendListBase::UAnimGraphNode_BlendListBase(const FObjectInitializer& ObjectInitializer)
+	: Super(ObjectInitializer)
 	, RemovedPinArrayIndex(INDEX_NONE)
 {
 }
diff --git a/Engine/Source/Editor/AnimGraph/Private/AnimGraphNode_BlendListByBool.cpp b/Engine/Source/Editor/AnimGraph/Private/AnimGraphNode_BlendListByBool.cpp
--- a/Engine/Source/Editor/AnimGraph/Private/AnimGraphNode_BlendListByBool.cpp
+++ b/Engine/Source/Editor/AnimGraph/Private/AnimGraphNode_BlendListByBool.cpp
@@ -11,8 +11,8 @@
 
 #define LOCTEXT_NAMESPACE "A3Nodes"
 
-UAnimGraphNode_BlendListByBool::UAnimGraphNode_BlendListByBool(const FPostConstructInitializeProperties& PCIP)
-	: Super(PCIP)
+UAnimGraphNode_BlendListByBool::UAnimGraphNode_BlendListByBool(const FObjectInitializer& ObjectInitializer)
+	: Super(ObjectInitializer)
 {
 }
 
diff --git a/Engine/Source/Editor/AnimGraph/Private/AnimGraphNode_BlendListByEnum.cpp b/Engine/Source/Editor/AnimGraph/Private/AnimGraphNode_BlendListByEnum.cpp
--- a/Engine/Source/Editor/AnimGraph/Private/AnimGraphNode_BlendListByEnum.cpp
+++ b/Engine/Source/Editor/AnimGraph/Private/AnimGraphNode_BlendListByEnum.cpp
@@ -16,8 +16,8 @@
 /////////////////////////////////////////////////////
 // UAnimGraphNode_BlendListByEnum
 
-UAnimGraphNode_BlendListByEnum::UAnimGraphNode_BlendListByEnum(const FPostConstructInitializeProperties& PCIP)
-	: Super(PCIP)
+UAnimGraphNode_BlendListByEnum::UAnimGraphNode_BlendListByEnum(const FObjectInitializer& ObjectInitializer)
+	: Super(ObjectInitializer)
 {
 }
 
diff --git a/Engine/Source/Editor/AnimGraph/Private/AnimGraphNode_BlendListByInt.cpp b/Engine/Source/Editor/AnimGraph/Private/AnimGraphNode_BlendListByInt.cpp
--- a/Engine/Source/Editor/AnimGraph/Private/AnimGraphNode_BlendListByInt.cpp
+++ b/Engine/Source/Editor/AnimGraph/Private/AnimGraphNode_BlendListByInt.cpp
@@ -12,8 +12,8 @@
 
 #define LOCTEXT_NAMESPACE "A3Nodes"
 
-UAnimGraphNode_BlendListByInt::UAnimGraphNode_BlendListByInt(const FPostConstructInitializeProperties& PCIP)
-	: Super(PCIP)
+UAnimGraphNode_BlendListByInt::UAnimGraphNode_BlendListByInt(const FObjectInitializer& ObjectInitializer)
+	: Super(ObjectInitializer)
 {
 }
 
diff --git a/Engine/Source/Editor/AnimGraph/Private/AnimGraphNode_BlendSpaceBase.cpp b/Engine/Source/Editor/AnimGraph/Private/AnimGraphNode_BlendSpaceBase.cpp
--- a/Engine/Source/Editor/AnimGraph/Private/AnimGraphNode_BlendSpaceBase.cpp
+++ b/Engine/Source/Editor/AnimGraph/Private/AnimGraphNode_BlendSpaceBase.cpp
@@ -44,8 +44,8 @@ struct FNewBlendSpacePlayerAction : public FEdGraphSchemaAction_K2NewNode
 /////////////////////////////////////////////////////
 // UAnimGraphNode_BlendSpaceBase
 
-UAnimGraphNode_BlendSpaceBase::UAnimGraphNode_BlendSpaceBase(const FPostConstructInitializeProperties& PCIP)
-	: Super(PCIP)
+UAnimGraphNode_BlendSpaceBase::UAnimGraphNode_BlendSpaceBase(const FObjectInitializer& ObjectInitializer)
+	: Super(ObjectInitializer)
 {
 }
 
diff --git a/Engine/Source/Editor/AnimGraph/Private/AnimGraphNode_BlendSpaceEvaluator.cpp b/Engine/Source/Editor/AnimGraph/Private/AnimGraphNode_BlendSpaceEvaluator.cpp
--- a/Engine/Source/Editor/AnimGraph/Private/AnimGraphNode_BlendSpaceEvaluator.cpp
+++ b/Engine/Source/Editor/AnimGraph/Private/AnimGraphNode_BlendSpaceEvaluator.cpp
@@ -10,8 +10,8 @@
 
 #define LOCTEXT_NAMESPACE "A3Nodes"
 
-UAnimGraphNode_BlendSpaceEvaluator::UAnimGraphNode_BlendSpaceEvaluator(const FPostConstructInitializeProperties& PCIP)
-	: Super(PCIP)
+UAnimGraphNode_BlendSpaceEvaluator::UAnimGraphNode_BlendSpaceEvaluator(const FObjectInitializer& ObjectInitializer)
+	: Super(ObjectInitializer)
 {
 }
 
diff --git a/Engine/Source/Editor/AnimGraph/Private/AnimGraphNode_BlendSpacePlayer.cpp b/Engine/Source/Editor/AnimGraph/Private/AnimGraphNode_BlendSpacePlayer.cpp
--- a/Engine/Source/Editor/AnimGraph/Private/AnimGraphNode_BlendSpacePlayer.cpp
+++ b/Engine/Source/Editor/AnimGraph/Private/AnimGraphNode_BlendSpacePlayer.cpp
@@ -12,8 +12,8 @@
 
 #define LOCTEXT_NAMESPACE "A3Nodes"
 
-UAnimGraphNode_BlendSpacePlayer::UAnimGraphNode_BlendSpacePlayer(const FPostConstructInitializeProperties& PCIP)
-	: Super(PCIP)
+UAnimGraphNode_BlendSpacePlayer::UAnimGraphNode_BlendSpacePlayer(const FObjectInitializer& ObjectInitializer)
+	: Super(ObjectInitializer)
 {
 }
 
diff --git a/Engine/Source/Editor/AnimGraph/Private/AnimGraphNode_BoneDrivenController.cpp b/Engine/Source/Editor/AnimGraph/Private/AnimGraphNode_BoneDrivenController.cpp
--- a/Engine/Source/Editor/AnimGraph/Private/AnimGraphNode_BoneDrivenController.cpp
+++ b/Engine/Source/Editor/AnimGraph/Private/AnimGraphNode_BoneDrivenController.cpp
@@ -5,8 +5,8 @@
 
 #define LOCTEXT_NAMESPACE "A3Nodes"
 
-UAnimGraphNode_BoneDrivenController::UAnimGraphNode_BoneDrivenController(const FPostConstructInitializeProperties& PCIP)
-	:Super(PCIP)
+UAnimGraphNode_BoneDrivenController::UAnimGraphNode_BoneDrivenController(const FObjectInitializer& ObjectInitializer)
+	:Super(ObjectInitializer)
 {
 
 }
diff --git a/Engine/Source/Editor/AnimGraph/Private/AnimGraphNode_ComponentToLocalSpace.cpp b/Engine/Source/Editor/AnimGraph/Private/AnimGraphNode_ComponentToLocalSpace.cpp
--- a/Engine/Source/Editor/AnimGraph/Private/AnimGraphNode_ComponentToLocalSpace.cpp
+++ b/Engine/Source/Editor/AnimGraph/Private/AnimGraphNode_ComponentToLocalSpace.cpp
@@ -8,8 +8,8 @@
 
 #define LOCTEXT_NAMESPACE "A3Nodes"
 
-UAnimGraphNode_ComponentToLocalSpace::UAnimGraphNode_ComponentToLocalSpace(const FPostConstructInitializeProperties& PCIP)
-	: Super(PCIP)
+UAnimGraphNode_ComponentToLocalSpace::UAnimGraphNode_ComponentToLocalSpace(const FObjectInitializer& ObjectInitializer)
+	: Super(ObjectInitializer)
 {
 }
 
diff --git a/Engine/Source/Editor/AnimGraph/Private/AnimGraphNode_CopyBone.cpp b/Engine/Source/Editor/AnimGraph/Private/AnimGraphNode_CopyBone.cpp
--- a/Engine/Source/Editor/AnimGraph/Private/AnimGraphNode_CopyBone.cpp
+++ b/Engine/Source/Editor/AnimGraph/Private/AnimGraphNode_CopyBone.cpp
@@ -8,8 +8,8 @@
 
 #define LOCTEXT_NAMESPACE "A3Nodes"
 
-UAnimGraphNode_CopyBone::UAnimGraphNode_CopyBone(const FPostConstructInitializeProperties& PCIP)
-	: Super(PCIP)
+UAnimGraphNode_CopyBone::UAnimGraphNode_CopyBone(const FObjectInitializer& ObjectInitializer)
+	: Super(ObjectInitializer)
 {
 }
 
diff --git a/Engine/Source/Editor/AnimGraph/Private/AnimGraphNode_CustomTransitionResult.cpp b/Engine/Source/Editor/AnimGraph/Private/AnimGraphNode_CustomTransitionResult.cpp
--- a/Engine/Source/Editor/AnimGraph/Private/AnimGraphNode_CustomTransitionResult.cpp
+++ b/Engine/Source/Editor/AnimGraph/Private/AnimGraphNode_CustomTransitionResult.cpp
@@ -10,8 +10,8 @@
 
 #define LOCTEXT_NAMESPACE "A3Nodes"
 
-UAnimGraphNode_CustomTransitionResult::UAnimGraphNode_CustomTransitionResult(const FPostConstructInitializeProperties& PCIP)
-	: Super(PCIP)
+UAnimGraphNode_CustomTransitionResult::UAnimGraphNode_CustomTransitionResult(const FObjectInitializer& ObjectInitializer)
+	: Super(ObjectInitializer)
 {
 }
 
diff --git a/Engine/Source/Editor/AnimGraph/Private/AnimGraphNode_Fabrik.cpp b/Engine/Source/Editor/AnimGraph/Private/AnimGraphNode_Fabrik.cpp
--- a/Engine/Source/Editor/AnimGraph/Private/AnimGraphNode_Fabrik.cpp
+++ b/Engine/Source/Editor/AnimGraph/Private/AnimGraphNode_Fabrik.cpp
@@ -8,8 +8,8 @@
 
 #define LOCTEXT_NAMESPACE "A3Nodes"
 
-UAnimGraphNode_Fabrik::UAnimGraphNode_Fabrik(const FPostConstructInitializeProperties& PCIP)
-: Super(PCIP)
+UAnimGraphNode_Fabrik::UAnimGraphNode_Fabrik(const FObjectInitializer& ObjectInitializer)
+: Super(ObjectInitializer)
 {
 }
 
diff --git a/Engine/Source/Editor/AnimGraph/Private/AnimGraphNode_HandIKRetargeting.cpp b/Engine/Source/Editor/AnimGraph/Private/AnimGraphNode_HandIKRetargeting.cpp
--- a/Engine/Source/Editor/AnimGraph/Private/AnimGraphNode_HandIKRetargeting.cpp
+++ b/Engine/Source/Editor/AnimGraph/Private/AnimGraphNode_HandIKRetargeting.cpp
@@ -8,8 +8,8 @@
 
 #define LOCTEXT_NAMESPACE "A3Nodes"
 
-UAnimGraphNode_HandIKRetargeting::UAnimGraphNode_HandIKRetargeting(const FPostConstructInitializeProperties& PCIP)
-: Super(PCIP)
+UAnimGraphNode_HandIKRetargeting::UAnimGraphNode_HandIKRetargeting(const FObjectInitializer& ObjectInitializer)
+: Super(ObjectInitializer)
 {
 }
 
diff --git a/Engine/Source/Editor/AnimGraph/Private/AnimGraphNode_IdentityPose.cpp b/Engine/Source/Editor/AnimGraph/Private/AnimGraphNode_IdentityPose.cpp
--- a/Engine/Source/Editor/AnimGraph/Private/AnimGraphNode_IdentityPose.cpp
+++ b/Engine/Source/Editor/AnimGraph/Private/AnimGraphNode_IdentityPose.cpp
@@ -8,8 +8,8 @@
 
 #define LOCTEXT_NAMESPACE "A3Nodes"
 
-UAnimGraphNode_IdentityPose::UAnimGraphNode_IdentityPose(const FPostConstructInitializeProperties& PCIP)
-	: Super(PCIP)
+UAnimGraphNode_IdentityPose::UAnimGraphNode_IdentityPose(const FObjectInitializer& ObjectInitializer)
+	: Super(ObjectInitializer)
 {
 	Node.RefPoseType = EIT_Additive;
 }
diff --git a/Engine/Source/Editor/AnimGraph/Private/AnimGraphNode_LayeredBoneBlend.cpp b/Engine/Source/Editor/AnimGraph/Private/AnimGraphNode_LayeredBoneBlend.cpp
--- a/Engine/Source/Editor/AnimGraph/Private/AnimGraphNode_LayeredBoneBlend.cpp
+++ b/Engine/Source/Editor/AnimGraph/Private/AnimGraphNode_LayeredBoneBlend.cpp
@@ -11,8 +11,8 @@
 
 #define LOCTEXT_NAMESPACE "A3Nodes"
 
-UAnimGraphNode_LayeredBoneBlend::UAnimGraphNode_LayeredBoneBlend(const FPostConstructInitializeProperties& PCIP)
-	: Super(PCIP)
+UAnimGraphNode_LayeredBoneBlend::UAnimGraphNode_LayeredBoneBlend(const FObjectInitializer& ObjectInitializer)
+	: Super(ObjectInitializer)
 {
 }
 
diff --git a/Engine/Source/Editor/AnimGraph/Private/AnimGraphNode_LocalRefPose.cpp b/Engine/Source/Editor/AnimGraph/Private/AnimGraphNode_LocalRefPose.cpp
--- a/Engine/Source/Editor/AnimGraph/Private/AnimGraphNode_LocalRefPose.cpp
+++ b/Engine/Source/Editor/AnimGraph/Private/AnimGraphNode_LocalRefPose.cpp
@@ -8,8 +8,8 @@
 
 #define LOCTEXT_NAMESPACE "A3Nodes"
 
-UAnimGraphNode_LocalRefPose::UAnimGraphNode_LocalRefPose(const FPostConstructInitializeProperties& PCIP)
-	: Super(PCIP)
+UAnimGraphNode_LocalRefPose::UAnimGraphNode_LocalRefPose(const FObjectInitializer& ObjectInitializer)
+	: Super(ObjectInitializer)
 {
 	Node.RefPoseType = EIT_LocalSpace;
 }
diff --git a/Engine/Source/Editor/AnimGraph/Private/AnimGraphNode_LocalToComponentSpace.cpp b/Engine/Source/Editor/AnimGraph/Private/AnimGraphNode_LocalToComponentSpace.cpp
--- a/Engine/Source/Editor/AnimGraph/Private/AnimGraphNode_LocalToComponentSpace.cpp
+++ b/Engine/Source/Editor/AnimGraph/Private/AnimGraphNode_LocalToComponentSpace.cpp
@@ -11,8 +11,8 @@
 
 #define LOCTEXT_NAMESPACE "A3Nodes"
 
-UAnimGraphNode_LocalToComponentSpace::UAnimGraphNode_LocalToComponentSpace(const FPostConstructInitializeProperties& PCIP)
-	: Super(PCIP)
+UAnimGraphNode_LocalToComponentSpace::UAnimGraphNode_LocalToComponentSpace(const FObjectInitializer& ObjectInitializer)
+	: Super(ObjectInitializer)
 {
 }
 
diff --git a/Engine/Source/Editor/AnimGraph/Private/AnimGraphNode_LookAt.cpp b/Engine/Source/Editor/AnimGraph/Private/AnimGraphNode_LookAt.cpp
--- a/Engine/Source/Editor/AnimGraph/Private/AnimGraphNode_LookAt.cpp
+++ b/Engine/Source/Editor/AnimGraph/Private/AnimGraphNode_LookAt.cpp
@@ -8,8 +8,8 @@
 
 #define LOCTEXT_NAMESPACE "AnimGraph_LookAt"
 
-UAnimGraphNode_LookAt::UAnimGraphNode_LookAt(const FPostConstructInitializeProperties& PCIP)
-	: Super(PCIP)
+UAnimGraphNode_LookAt::UAnimGraphNode_LookAt(const FObjectInitializer& ObjectInitializer)
+	: Super(ObjectInitializer)
 {
 }
 
diff --git a/Engine/Source/Editor/AnimGraph/Private/AnimGraphNode_MeshRefPose.cpp b/Engine/Source/Editor/AnimGraph/Private/AnimGraphNode_MeshRefPose.cpp
--- a/Engine/Source/Editor/AnimGraph/Private/AnimGraphNode_MeshRefPose.cpp
+++ b/Engine/Source/Editor/AnimGraph/Private/AnimGraphNode_MeshRefPose.cpp
@@ -9,8 +9,8 @@
 
 #define LOCTEXT_NAMESPACE "A3Nodes"
 
-UAnimGraphNode_MeshRefPose::UAnimGraphNode_MeshRefPose(const FPostConstructInitializeProperties& PCIP)
-	: Super(PCIP)
+UAnimGraphNode_MeshRefPose::UAnimGraphNode_MeshRefPose(const FObjectInitializer& ObjectInitializer)
+	: Super(ObjectInitializer)
 {
 }
 
diff --git a/Engine/Source/Editor/AnimGraph/Private/AnimGraphNode_ModifyBone.cpp b/Engine/Source/Editor/AnimGraph/Private/AnimGraphNode_ModifyBone.cpp
--- a/Engine/Source/Editor/AnimGraph/Private/AnimGraphNode_ModifyBone.cpp
+++ b/Engine/Source/Editor/AnimGraph/Private/AnimGraphNode_ModifyBone.cpp
@@ -8,8 +8,8 @@
 
 #define LOCTEXT_NAMESPACE "A3Nodes"
 
-UAnimGraphNode_ModifyBone::UAnimGraphNode_ModifyBone(const FPostConstructInitializeProperties& PCIP)
-	: Super(PCIP)
+UAnimGraphNode_ModifyBone::UAnimGraphNode_ModifyBone(const FObjectInitializer& ObjectInitializer)
+	: Super(ObjectInitializer)
 {
 }
 
diff --git a/Engine/Source/Editor/AnimGraph/Private/AnimGraphNode_RefPoseBase.cpp b/Engine/Source/Editor/AnimGraph/Private/AnimGraphNode_RefPoseBase.cpp
--- a/Engine/Source/Editor/AnimGraph/Private/AnimGraphNode_RefPoseBase.cpp
+++ b/Engine/Source/Editor/AnimGraph/Private/AnimGraphNode_RefPoseBase.cpp
@@ -6,8 +6,8 @@
 /////////////////////////////////////////////////////
 // UAnimGraphNode_RefPoseBase
 
-UAnimGraphNode_RefPoseBase::UAnimGraphNode_RefPoseBase(const FPostConstructInitializeProperties& PCIP)
-	: Super(PCIP)
+UAnimGraphNode_RefPoseBase::UAnimGraphNode_RefPoseBase(const FObjectInitializer& ObjectInitializer)
+	: Super(ObjectInitializer)
 {
 }
 
diff --git a/Engine/Source/Editor/AnimGraph/Private/AnimGraphNode_Root.cpp b/Engine/Source/Editor/AnimGraph/Private/AnimGraphNode_Root.cpp
--- a/Engine/Source/Editor/AnimGraph/Private/AnimGraphNode_Root.cpp
+++ b/Engine/Source/Editor/AnimGraph/Private/AnimGraphNode_Root.cpp
@@ -38,8 +38,8 @@ void FPoseLinkMappingRecord::PatchLinkIndex(uint8* DestinationPtr, int32 LinkID,
 
 #define LOCTEXT_NAMESPACE "A3Nodes"
 
-UAnimGraphNode_Root::UAnimGraphNode_Root(const FPostConstructInitializeProperties& PCIP)
-	: Super(PCIP)
+UAnimGraphNode_Root::UAnimGraphNode_Root(const FObjectInitializer& ObjectInitializer)
+	: Super(ObjectInitializer)
 {
 }
 
diff --git a/Engine/Source/Editor/AnimGraph/Private/AnimGraphNode_RotateRootBone.cpp b/Engine/Source/Editor/AnimGraph/Private/AnimGraphNode_RotateRootBone.cpp
--- a/Engine/Source/Editor/AnimGraph/Private/AnimGraphNode_RotateRootBone.cpp
+++ b/Engine/Source/Editor/AnimGraph/Private/AnimGraphNode_RotateRootBone.cpp
@@ -11,8 +11,8 @@
 
 #define LOCTEXT_NAMESPACE "A3Nodes"
 
-UAnimGraphNode_RotateRootBone::UAnimGraphNode_RotateRootBone(const FPostConstructInitializeProperties& PCIP)
-	: Super(PCIP)
+UAnimGraphNode_RotateRootBone::UAnimGraphNode_RotateRootBone(const FObjectInitializer& ObjectInitializer)
+	: Super(ObjectInitializer)
 {
 }
 
diff --git a/Engine/Source/Editor/AnimGraph/Private/AnimGraphNode_RotationMultiplier.cpp b/Engine/Source/Editor/AnimGraph/Private/AnimGraphNode_RotationMultiplier.cpp
--- a/Engine/Source/Editor/AnimGraph/Private/AnimGraphNode_RotationMultiplier.cpp
+++ b/Engine/Source/Editor/AnimGraph/Private/AnimGraphNode_RotationMultiplier.cpp
@@ -7,8 +7,8 @@
 
 #define LOCTEXT_NAMESPACE "A3Nodes"
 
-UAnimGraphNode_RotationMultiplier::UAnimGraphNode_RotationMultiplier(const FPostConstructInitializeProperties& PCIP)
-	: Super(PCIP)
+UAnimGraphNode_RotationMultiplier::UAnimGraphNode_RotationMultiplier(const FObjectInitializer& ObjectInitializer)
+	: Super(ObjectInitializer)
 {
 }
 
diff --git a/Engine/Source/Editor/AnimGraph/Private/AnimGraphNode_RotationOffsetBlendSpace.cpp b/Engine/Source/Editor/AnimGraph/Private/AnimGraphNode_RotationOffsetBlendSpace.cpp
--- a/Engine/Source/Editor/AnimGraph/Private/AnimGraphNode_RotationOffsetBlendSpace.cpp
+++ b/Engine/Source/Editor/AnimGraph/Private/AnimGraphNode_RotationOffsetBlendSpace.cpp
@@ -12,8 +12,8 @@
 
 #define LOCTEXT_NAMESPACE "A3Nodes"
 
-UAnimGraphNode_RotationOffsetBlendSpace::UAnimGraphNode_RotationOffsetBlendSpace(const FPostConstructInitializeProperties& PCIP)
-	: Super(PCIP)
+UAnimGraphNode_RotationOffsetBlendSpace::UAnimGraphNode_RotationOffsetBlendSpace(const FObjectInitializer& ObjectInitializer)
+	: Super(ObjectInitializer)
 {
 }
 
