Fix importing speedtrees attempting to generate lightmap uv's instead of using the ones in the SRT file.

diff --git a/Engine/Plugins/Editor/SpeedTreeImporter/Source/SpeedTreeImporter/Private/SpeedTreeImportFactory.cpp b/Engine/Plugins/Editor/SpeedTreeImporter/Source/SpeedTreeImporter/Private/SpeedTreeImportFactory.cpp
--- a/Engine/Plugins/Editor/SpeedTreeImporter/Source/SpeedTreeImporter/Private/SpeedTreeImportFactory.cpp
+++ b/Engine/Plugins/Editor/SpeedTreeImporter/Source/SpeedTreeImporter/Private/SpeedTreeImportFactory.cpp
@@ -1343,6 +1343,7 @@ UObject* USpeedTreeImportFactory::FactoryCreateBinary(UClass* InClass, UObject*
 						LODModel->BuildSettings.bRecomputeTangents = false;
 						LODModel->BuildSettings.bRemoveDegenerates = true;
 						LODModel->BuildSettings.bUseFullPrecisionUVs = false;	
+						LODModel->BuildSettings.bGenerateLightmapUVs = false;
 						LODModel->ScreenSize = 0.1f / FMath::Pow(2.0f, StaticMesh->SourceModels.Num() - 1);
 						LODModel->RawMeshBulkData->SaveRawMesh(RawMesh);
 
@@ -1454,6 +1455,7 @@ UObject* USpeedTreeImportFactory::FactoryCreateBinary(UClass* InClass, UObject*
 					LODModel->BuildSettings.bRecomputeTangents = false;
 					LODModel->BuildSettings.bRemoveDegenerates = true;
 					LODModel->BuildSettings.bUseFullPrecisionUVs = false;
+					LODModel->BuildSettings.bGenerateLightmapUVs = false;
 					LODModel->ScreenSize = 0.1f / FMath::Pow(2.0f, StaticMesh->SourceModels.Num() - 1);
 					LODModel->RawMeshBulkData->SaveRawMesh(RawMesh);
 				}
