Integrate from GregC@IDV

SpeedTree importer improvements:
- two sided lighting fix included even when there is no normal map
- added lod option for painted foliage
- reimport maintains changed materials (as long as the material setup hasn't changed drastically)
- reimport touches textures to make sure they reimport, too

diff --git a/Engine/Plugins/Editor/SpeedTreeImporter/Source/SpeedTreeImporter/Classes/SpeedTreeImportFactory.h b/Engine/Plugins/Editor/SpeedTreeImporter/Source/SpeedTreeImporter/Classes/SpeedTreeImportFactory.h
--- a/Engine/Plugins/Editor/SpeedTreeImporter/Source/SpeedTreeImporter/Classes/SpeedTreeImportFactory.h
+++ b/Engine/Plugins/Editor/SpeedTreeImporter/Source/SpeedTreeImporter/Classes/SpeedTreeImportFactory.h
@@ -19,5 +19,7 @@ class USpeedTreeImportFactory : public UFactory
 	virtual UClass* ResolveSupportedClass() override;
 	virtual UObject* FactoryCreateBinary(UClass* InClass, UObject* InParent, FName InName, EObjectFlags Flags, UObject* Context, const TCHAR* Type, const uint8*& Buffer, const uint8* BufferEnd, FFeedbackContext* Warn) override;
 #endif
+
+	TSet<UPackage*> LoadedPackages;
 };
 
diff --git a/Engine/Plugins/Editor/SpeedTreeImporter/Source/SpeedTreeImporter/Private/SpeedTreeImportFactory.cpp b/Engine/Plugins/Editor/SpeedTreeImporter/Source/SpeedTreeImporter/Private/SpeedTreeImportFactory.cpp
--- a/Engine/Plugins/Editor/SpeedTreeImporter/Source/SpeedTreeImporter/Private/SpeedTreeImportFactory.cpp
+++ b/Engine/Plugins/Editor/SpeedTreeImporter/Source/SpeedTreeImporter/Private/SpeedTreeImportFactory.cpp
@@ -8,6 +8,7 @@
 #include "Materials/MaterialExpressionClamp.h"
 #include "Materials/MaterialExpressionComponentMask.h"
 #include "Materials/MaterialExpressionConstant.h"
+#include "Materials/MaterialExpressionConstant3Vector.h"
 #include "Materials/MaterialExpressionLinearInterpolate.h"
 #include "Materials/MaterialExpressionPower.h"
 #include "Materials/MaterialExpressionMultiply.h"
@@ -67,6 +68,14 @@ public:
 	/** Tree scale **/
 	float							TreeScale;
 
+	/** LOD type **/
+	enum EImportLODType
+	{
+		ILT_PaintedFoliage,
+		ILT_IndividualActors
+	};
+	EImportLODType					LODType;
+
 	/** options */
 	TSharedPtr<SCheckBox>			MakeMaterialsCheck;
 	TSharedPtr<SCheckBox>			IncludeNormalMapCheck;
@@ -97,6 +106,7 @@ public:
 	SSpeedTreeImportOptions() :
 		ImportGeometryType(IGT_3D),
 		TreeScale(30.48f),
+		LODType(ILT_PaintedFoliage),
 		bImport(false)
 	{}
 
@@ -115,20 +125,20 @@ public:
 	{
 		return SNew(SCheckBox)
 		.Style(FEditorStyle::Get(), "RadioButton")
-		.IsChecked(this, &SSpeedTreeImportOptions::IsRadioChecked, Type)
-		.OnCheckStateChanged(this, &SSpeedTreeImportOptions::OnRadioChanged, Type)
+		.IsChecked(this, &SSpeedTreeImportOptions::IsGeometryTypeRadioChecked, Type)
+		.OnCheckStateChanged(this, &SSpeedTreeImportOptions::OnGeometryTypeRadioChanged, Type)
 		.Content()
 		[
 			SNew(STextBlock).Text(NameText)
 		];
 	}
 
-	ECheckBoxState IsRadioChecked(EImportGeometryType ButtonId) const
+	ECheckBoxState IsGeometryTypeRadioChecked(EImportGeometryType ButtonId) const
 	{
 		return (ImportGeometryType == ButtonId) ? ECheckBoxState::Checked : ECheckBoxState::Unchecked;
 	}
 
-	void OnRadioChanged( ECheckBoxState NewRadioState, EImportGeometryType RadioThatChanged )
+	void OnGeometryTypeRadioChanged(ECheckBoxState NewRadioState, EImportGeometryType RadioThatChanged )
 	{
 		if (NewRadioState == ECheckBoxState::Checked)
 		{
@@ -136,6 +146,31 @@ public:
 		}
 	}
 
+	TSharedRef<SWidget> CreateLodTypeCombo(const FText& NameText, bool bChecked, EImportLODType Type)
+	{
+		return SNew(SCheckBox)
+		.Style(FEditorStyle::Get(), "RadioButton")
+		.IsChecked(this, &SSpeedTreeImportOptions::IsLodTypeRadioChecked, Type)
+		.OnCheckStateChanged(this, &SSpeedTreeImportOptions::OnLodTypeRadioChanged, Type)
+		.Content()
+		[
+			SNew(STextBlock).Text(NameText)
+		];
+	}
+
+	ECheckBoxState IsLodTypeRadioChecked(EImportLODType ButtonId) const
+	{
+		return (LODType == ButtonId) ? ECheckBoxState::Checked : ECheckBoxState::Unchecked;
+	}
+
+	void OnLodTypeRadioChanged(ECheckBoxState NewRadioState, EImportLODType RadioThatChanged)
+	{
+		if (NewRadioState == ECheckBoxState::Checked)
+		{
+			LODType = RadioThatChanged;
+		}
+	}
+
 	void Construct(const FArguments& InArgs)
 	{
 		WidgetWindow = InArgs._WidgetWindow;
@@ -150,6 +185,24 @@ public:
 				SNew(SVerticalBox)
 
 				// options
+				+SVerticalBox::Slot().AutoHeight().Padding(5)
+				[
+					SNew(SVerticalBox)
+
+					+ SVerticalBox::Slot().AutoHeight().Padding(5).HAlign(HAlign_Left)
+					[
+						SNew(STextBlock).Text(LOCTEXT("TreeScaleLabel", "Tree Scale"))
+					]
+
+					+SVerticalBox::Slot().AutoHeight().Padding(5)
+					[
+						SNew(SEditableTextBox)
+						.SelectAllTextWhenFocused(true)
+						.OnTextCommitted(this, &SSpeedTreeImportOptions::ScaleTextCommitted)
+						.Text(FText::AsNumber(TreeScale))
+					]
+				]
+
 				+SVerticalBox::Slot().AutoHeight().Padding(5)
 				[
 					SNew(SVerticalBox)
@@ -165,7 +218,7 @@ public:
 					
 						+SVerticalBox::Slot().AutoHeight().Padding(5)
 						[
-							CreateGeometryTypeCombo( LOCTEXT("3D_LODs", "3D LODs"), true, IGT_3D)
+							CreateGeometryTypeCombo(LOCTEXT("3D_LODs", "3D LODs"), true, IGT_3D)
 						]
 
 						+SVerticalBox::Slot().AutoHeight().Padding(5)
@@ -186,15 +239,22 @@ public:
 
 					+ SVerticalBox::Slot().AutoHeight().Padding(5).HAlign(HAlign_Left)
 					[
-						SNew(STextBlock).Text(LOCTEXT("TreeScaleLabel", "Tree Scale"))
+						SNew(STextBlock).Text(LOCTEXT("LODSetupLabel", "LOD Setup"))
 					]
 
-					+SVerticalBox::Slot().AutoHeight().Padding(5)
+					+SVerticalBox::Slot().AutoHeight().Padding(15, 5)
 					[
-						SNew(SEditableTextBox)
-						.SelectAllTextWhenFocused(true)
-						.OnTextCommitted(this, &SSpeedTreeImportOptions::ScaleTextCommitted)
-						.Text(FText::AsNumber(TreeScale))
+						SNew(SVerticalBox)
+						
+						+SVerticalBox::Slot().AutoHeight().Padding(5)
+						[
+							CreateLodTypeCombo(LOCTEXT("PaintedFoliage", "Painted Foliage"), true, ILT_PaintedFoliage)
+						]
+
+						+SVerticalBox::Slot().AutoHeight().Padding(5)
+						[
+							CreateLodTypeCombo(LOCTEXT("IndividualActors", "Individual Actors"), false, ILT_IndividualActors)
+						]
 					]
 				]
 
@@ -354,7 +414,7 @@ USpeedTreeImportFactory::USpeedTreeImportFactory(const FObjectInitializer& Objec
 
 FText USpeedTreeImportFactory::GetDisplayName() const
 {
-	return LOCTEXT("SpeedTreeImportFactoryDescription", "Speed Tree");
+	return LOCTEXT("SpeedTreeImportFactoryDescription", "SpeedTree");
 }
 
 
@@ -370,7 +430,7 @@ UClass* USpeedTreeImportFactory::ResolveSupportedClass()
 	return UStaticMesh::StaticClass();
 }
 
-static UTexture* CreateSpeedTreeMaterialTexture(UObject* Parent, FString Filename, bool bNormalMap)
+UTexture* CreateSpeedTreeMaterialTexture(UObject* Parent, FString Filename, bool bNormalMap, TSet<UPackage*>& LoadedPackages)
 {
 	UTexture* UnrealTexture = NULL;
 
@@ -388,10 +448,19 @@ static UTexture* CreateSpeedTreeMaterialTexture(UObject* Parent, FString Filenam
 	NewPackageName = PackageTools::SanitizePackageName(NewPackageName);
 	UPackage* Package = CreatePackage(NULL, *NewPackageName);
 
-	// does not override existing textures
+	// find existing texture
 	UnrealTexture = FindObject<UTexture>(Package, *TextureName);
 	if (UnrealTexture != NULL)
 	{
+		if (!LoadedPackages.Contains(Package))
+		{
+			LoadedPackages.Add(Package);
+			if (!FReimportManager::Instance()->Reimport(UnrealTexture, true))
+			{
+				UE_LOG(LogSpeedTreeImport, Warning, TEXT("Manual texture reimport and recompression may be needed for %s"), *TextureName);
+			}
+		}
+
 		return UnrealTexture;
 	}
 
@@ -405,6 +474,7 @@ static UTexture* CreateSpeedTreeMaterialTexture(UObject* Parent, FString Filenam
 	else
 	{
 		UTextureFactory* TextureFact = new UTextureFactory(FObjectInitializer());
+		TextureFact->AddToRoot();
 		TextureFact->SuppressImportOverwriteDialog();
 
 		if (bNormalMap)
@@ -417,18 +487,24 @@ static UTexture* CreateSpeedTreeMaterialTexture(UObject* Parent, FString Filenam
 		UnrealTexture = (UTexture*)TextureFact->FactoryCreateBinary(UTexture2D::StaticClass(), Package, *TextureName, RF_Standalone|RF_Public, NULL, *Extension, PtrTexture, PtrTexture + TextureData.Num(), GWarn);
 		if (UnrealTexture != NULL)
 		{
+			UnrealTexture->SourceFilePath = FReimportManager::SanitizeImportFilename(Filename, UnrealTexture);
+			UnrealTexture->SourceFileTimestamp = IFileManager::Get().GetTimeStamp(*Filename).ToString();
+
 			// Notify the asset registry
 			FAssetRegistryModule::AssetCreated(UnrealTexture);
 
 			// Set the dirty flag so this package will get saved later
 			Package->SetDirtyFlag(true);
+			LoadedPackages.Add(Package);
 		}
+
+		TextureFact->RemoveFromRoot();
 	}
 
 	return UnrealTexture;
 }
 
-static void LayoutMaterial(UMaterialInterface* MaterialInterface)
+void LayoutMaterial(UMaterialInterface* MaterialInterface)
 {
 	UMaterial* Material = MaterialInterface->GetMaterial();
 	Material->EditorX = 0;
@@ -535,7 +611,7 @@ static void LayoutMaterial(UMaterialInterface* MaterialInterface)
 	}
 }
 
-static UMaterialInterface* CreateSpeedTreeMaterial(UObject* Parent, FString MaterialFullName, const SpeedTree::SRenderState* RenderState, TSharedPtr<SSpeedTreeImportOptions> Options, ESpeedTreeWindType WindType, int NumBillboards)
+UMaterialInterface* CreateSpeedTreeMaterial(UObject* Parent, FString MaterialFullName, const SpeedTree::SRenderState* RenderState, TSharedPtr<SSpeedTreeImportOptions> Options, ESpeedTreeWindType WindType, int NumBillboards, TSet<UPackage*>& LoadedPackages)
 {
 	// Make sure we have a parent
 	if (!Options->MakeMaterialsCheck->IsChecked() || !ensure(Parent))
@@ -553,26 +629,38 @@ static UMaterialInterface* CreateSpeedTreeMaterial(UObject* Parent, FString Mate
 	UMaterialInterface* UnrealMaterialInterface = FindObject<UMaterialInterface>(Package, *FixedMaterialName);
 	if (UnrealMaterialInterface != NULL)
 	{
+		// touch the textures anyway to make sure they reload if necessary
+		UTexture* DiffuseTexture = CreateSpeedTreeMaterialTexture(Parent, ANSI_TO_TCHAR(RenderState->m_apTextures[SpeedTree::TL_DIFFUSE]), false, LoadedPackages);
+		if (DiffuseTexture)
+		{
+			if (RenderState->m_bBranchesPresent && Options->IncludeDetailMapCheck->IsChecked())
+			{
+				UTexture* DetailTexture = CreateSpeedTreeMaterialTexture(Parent, ANSI_TO_TCHAR(RenderState->m_apTextures[SpeedTree::TL_DETAIL_DIFFUSE]), false, LoadedPackages);
+			}
+		}
+		if (Options->IncludeSpecularMapCheck->IsChecked())
+		{
+			UTexture* SpecularTexture = CreateSpeedTreeMaterialTexture(Parent, ANSI_TO_TCHAR(RenderState->m_apTextures[SpeedTree::TL_SPECULAR_MASK]), false, LoadedPackages);
+		}
+		if (Options->IncludeNormalMapCheck->IsChecked())
+		{
+			UTexture* NormalTexture = CreateSpeedTreeMaterialTexture(Parent, ANSI_TO_TCHAR(RenderState->m_apTextures[SpeedTree::TL_NORMAL]), true, LoadedPackages);
+		}
+
 		return UnrealMaterialInterface;
 	}
 	
 	// create an unreal material asset
 	UMaterialFactoryNew* MaterialFactory = new UMaterialFactoryNew(FObjectInitializer());
 	UMaterial* UnrealMaterial = (UMaterial*)MaterialFactory->FactoryCreateNew(UMaterial::StaticClass(), Package, *FixedMaterialName, RF_Standalone|RF_Public, NULL, GWarn);
-	if (UnrealMaterial != NULL)
-	{
-		if (!RenderState->m_bDiffuseAlphaMaskIsOpaque && !RenderState->m_bBranchesPresent && !RenderState->m_bRigidMeshesPresent)
-		{
-			UnrealMaterial->BlendMode = BLEND_Masked;
-			UnrealMaterial->SetCastShadowAsMasked(true);
-			UnrealMaterial->TwoSided = !(RenderState->m_bHorzBillboard || RenderState->m_bVertBillboard);
-		}
-
-		// Notify the asset registry
-		FAssetRegistryModule::AssetCreated(UnrealMaterial);
+	FAssetRegistryModule::AssetCreated(UnrealMaterial);
+	Package->SetDirtyFlag(true);
 
-		// Set the dirty flag so this package will get saved later
-		Package->SetDirtyFlag(true);
+	if (!RenderState->m_bDiffuseAlphaMaskIsOpaque && !RenderState->m_bBranchesPresent && !RenderState->m_bRigidMeshesPresent)
+	{
+		UnrealMaterial->BlendMode = BLEND_Masked;
+		UnrealMaterial->SetCastShadowAsMasked(true);
+		UnrealMaterial->TwoSided = !(RenderState->m_bHorzBillboard || RenderState->m_bVertBillboard);
 	}
 
 	UMaterialExpressionClamp* BranchSeamAmount = NULL;
@@ -601,7 +689,7 @@ static UMaterialInterface* CreateSpeedTreeMaterial(UObject* Parent, FString Mate
 	}
 
 	// textures and properties
-	UTexture* DiffuseTexture = CreateSpeedTreeMaterialTexture(Parent, ANSI_TO_TCHAR(RenderState->m_apTextures[SpeedTree::TL_DIFFUSE]), false);
+	UTexture* DiffuseTexture = CreateSpeedTreeMaterialTexture(Parent, ANSI_TO_TCHAR(RenderState->m_apTextures[SpeedTree::TL_DIFFUSE]), false, LoadedPackages);
 	if (DiffuseTexture)
 	{
 		// make texture sampler
@@ -643,7 +731,7 @@ static UMaterialInterface* CreateSpeedTreeMaterial(UObject* Parent, FString Mate
 
 		if (RenderState->m_bBranchesPresent && Options->IncludeDetailMapCheck->IsChecked())
 		{
-			UTexture* DetailTexture = CreateSpeedTreeMaterialTexture(Parent, ANSI_TO_TCHAR(RenderState->m_apTextures[SpeedTree::TL_DETAIL_DIFFUSE]), false);
+			UTexture* DetailTexture = CreateSpeedTreeMaterialTexture(Parent, ANSI_TO_TCHAR(RenderState->m_apTextures[SpeedTree::TL_DETAIL_DIFFUSE]), false, LoadedPackages);
 			if (DetailTexture)
 			{
 				// add/find UVSet
@@ -679,7 +767,7 @@ static UMaterialInterface* CreateSpeedTreeMaterial(UObject* Parent, FString Mate
 	bool bMadeSpecular = false;
 	if (Options->IncludeSpecularMapCheck->IsChecked())
 	{
-		UTexture* SpecularTexture = CreateSpeedTreeMaterialTexture(Parent, ANSI_TO_TCHAR(RenderState->m_apTextures[SpeedTree::TL_SPECULAR_MASK]), false);
+		UTexture* SpecularTexture = CreateSpeedTreeMaterialTexture(Parent, ANSI_TO_TCHAR(RenderState->m_apTextures[SpeedTree::TL_SPECULAR_MASK]), false, LoadedPackages);
 		if (SpecularTexture)
 		{
 			// make texture sampler
@@ -703,7 +791,7 @@ static UMaterialInterface* CreateSpeedTreeMaterial(UObject* Parent, FString Mate
 
 	if (Options->IncludeNormalMapCheck->IsChecked())
 	{
-		UTexture* NormalTexture = CreateSpeedTreeMaterialTexture(Parent, ANSI_TO_TCHAR(RenderState->m_apTextures[SpeedTree::TL_NORMAL]), true);
+		UTexture* NormalTexture = CreateSpeedTreeMaterialTexture(Parent, ANSI_TO_TCHAR(RenderState->m_apTextures[SpeedTree::TL_NORMAL]), true, LoadedPackages);
 		if (NormalTexture)
 		{
 			// make texture sampler
@@ -777,15 +865,27 @@ static UMaterialInterface* CreateSpeedTreeMaterial(UObject* Parent, FString Mate
 	}
 
 	// UE4 flips normals for two-sided materials. SpeedTrees don't need that
-	if (UnrealMaterial->TwoSided && UnrealMaterial->Normal.Expression != NULL)
+	if (UnrealMaterial->TwoSided)
 	{
 		UMaterialExpressionTwoSidedSign* TwoSidedSign = ConstructObject<UMaterialExpressionTwoSidedSign>(UMaterialExpressionTwoSidedSign::StaticClass(), UnrealMaterial);
 		UnrealMaterial->Expressions.Add(TwoSidedSign);
 
 		UMaterialExpressionMultiply* Multiply = ConstructObject<UMaterialExpressionMultiply>(UMaterialExpressionMultiply::StaticClass(), UnrealMaterial);
-		Multiply->A.Expression = UnrealMaterial->Normal.Expression;
-		Multiply->B.Expression = TwoSidedSign;
 		UnrealMaterial->Expressions.Add(Multiply);
+		Multiply->A.Expression = TwoSidedSign;
+
+		if (UnrealMaterial->Normal.Expression == NULL)
+		{
+			UMaterialExpressionConstant3Vector* VertexNormalExpression = ConstructObject<UMaterialExpressionConstant3Vector>(UMaterialExpressionConstant3Vector::StaticClass(), UnrealMaterial);
+			UnrealMaterial->Expressions.Add(VertexNormalExpression);
+			VertexNormalExpression->Constant = FLinearColor(0.0f, 0.0f, 1.0f);
+
+			Multiply->B.Expression = VertexNormalExpression;
+		}
+		else
+		{
+			Multiply->B.Expression = UnrealMaterial->Normal.Expression;
+		}
 
 		UnrealMaterial->Normal.Expression = Multiply;
 	}
@@ -987,6 +1087,8 @@ UObject* USpeedTreeImportFactory::FactoryCreateBinary(UClass* InClass, UObject*
 			}
 			else
 			{
+				LoadedPackages.Empty( );
+
 				// make static mesh object
 				FString MeshName = ObjectTools::SanitizeObjectName(InName.ToString());
 				FString NewPackageName = FPackageName::GetLongPackagePath(InParent->GetOutermost()->GetName()) + TEXT("/") + MeshName;
@@ -994,9 +1096,19 @@ UObject* USpeedTreeImportFactory::FactoryCreateBinary(UClass* InClass, UObject*
 				UPackage* Package = CreatePackage(NULL, *NewPackageName);
 				
 				// clear out old mesh
+				TArray<UMaterialInterface*> OldMaterials;
 				UStaticMesh* ExistingMesh = FindObject<UStaticMesh>(Package, *MeshName);
+				FGlobalComponentReregisterContext RecreateComponents;
+
 				if (ExistingMesh)
 				{
+					OldMaterials = ExistingMesh->Materials;
+					for (int32 i = 0; i < OldMaterials.Num(); ++i)
+					{
+						OldMaterials[i]->PreEditChange(NULL);
+						OldMaterials[i]->PostEditChange();
+					}
+
 					// Free any RHI resources for existing mesh before we re-create in place.
 					ExistingMesh->PreEditChange(NULL);
 				}
@@ -1149,7 +1261,7 @@ UObject* USpeedTreeImportFactory::FactoryCreateBinary(UClass* InClass, UObject*
 
 								MaterialName = ObjectTools::SanitizeObjectName(MaterialName);
 
-								UMaterialInterface* Material = CreateSpeedTreeMaterial(InParent, MaterialName, RenderState, Options, WindType, SpeedTreeGeometry->m_sVertBBs.m_nNumBillboards);
+								UMaterialInterface* Material = CreateSpeedTreeMaterial(InParent, MaterialName, RenderState, Options, WindType, SpeedTreeGeometry->m_sVertBBs.m_nNumBillboards, LoadedPackages);
 								
 								RenderStateIndexToStaticMeshIndex.Add(DrawCall->m_nRenderStateIndex, StaticMesh->Materials.Num());
 								MaterialIndex = StaticMesh->Materials.Num();
@@ -1361,7 +1473,7 @@ UObject* USpeedTreeImportFactory::FactoryCreateBinary(UClass* InClass, UObject*
 				// make billboard LOD
 				if (Options->ImportGeometryType != SSpeedTreeImportOptions::IGT_3D && SpeedTreeGeometry->m_sVertBBs.m_nNumBillboards > 0)
 				{
-					UMaterialInterface* Material = CreateSpeedTreeMaterial(InParent, MeshName + "_Billboard", &SpeedTreeGeometry->m_aBillboardRenderStates[SpeedTree::RENDER_PASS_MAIN], Options, WindType, SpeedTreeGeometry->m_sVertBBs.m_nNumBillboards);
+					UMaterialInterface* Material = CreateSpeedTreeMaterial(InParent, MeshName + "_Billboard", &SpeedTreeGeometry->m_aBillboardRenderStates[SpeedTree::RENDER_PASS_MAIN], Options, WindType, SpeedTreeGeometry->m_sVertBBs.m_nNumBillboards, LoadedPackages);
 					int32 MaterialIndex = StaticMesh->Materials.Num();
 					StaticMesh->Materials.Add(Material);
 
@@ -1462,6 +1574,11 @@ UObject* USpeedTreeImportFactory::FactoryCreateBinary(UClass* InClass, UObject*
 					LODModel->RawMeshBulkData->SaveRawMesh(RawMesh);
 				}
 
+				if (OldMaterials.Num() == StaticMesh->Materials.Num())
+				{
+					StaticMesh->Materials = OldMaterials;
+				}
+
 				StaticMesh->Build();
 
 				if (Options->IncludeCollision->IsChecked())
@@ -1475,7 +1592,10 @@ UObject* USpeedTreeImportFactory::FactoryCreateBinary(UClass* InClass, UObject*
 				}
 
 				// make better LOD info for SpeedTrees
-				StaticMesh->bAutoComputeLODScreenSize = false;
+				if (Options->LODType == SSpeedTreeImportOptions::ILT_IndividualActors)
+				{
+					StaticMesh->bAutoComputeLODScreenSize = false;
+				}
 				StaticMesh->bRequiresLODDistanceConversion = false;
 			}
 		}
