[INTEGRATE] Change 2420267 by [REDACTED_PERSON][REDACTED_HANDLE] on 2015/01/27 09:27:47

	Fixed slow compile times of SpeedTreeImportFactory.cpp
	- Factored out inner loop that optimizer was stumbling on (135s => 11s)

diff --git a/Engine/Plugins/Editor/SpeedTreeImporter/Source/SpeedTreeImporter/Private/SpeedTreeImportFactory.cpp b/Engine/Plugins/Editor/SpeedTreeImporter/Source/SpeedTreeImporter/Private/SpeedTreeImportFactory.cpp
--- a/Engine/Plugins/Editor/SpeedTreeImporter/Source/SpeedTreeImporter/Private/SpeedTreeImportFactory.cpp
+++ b/Engine/Plugins/Editor/SpeedTreeImporter/Source/SpeedTreeImporter/Private/SpeedTreeImportFactory.cpp
@@ -52,6 +52,7 @@
 DEFINE_LOG_CATEGORY_STATIC(LogSpeedTreeImport, Log, All);
 
 /** UI to pick options when importing  SpeedTree */
+BEGIN_SLATE_FUNCTION_BUILD_OPTIMIZATION
 class SSpeedTreeImportOptions : public SCompoundWidget
 {
 public:
@@ -394,6 +395,7 @@ public:
 		TTypeFromString<float>::FromString(TreeScale, *(CommentText.ToString()));
 	}
 };
+END_SLATE_FUNCTION_BUILD_OPTIMIZATION
 
 //////////////////////////////////////////////////////////////////////////
 
@@ -1042,7 +1044,123 @@ static void MakeBodyFromCollisionObjects(UStaticMesh* StaticMesh, const SpeedTre
 }
 
 
-UObject* USpeedTreeImportFactory::FactoryCreateBinary(UClass* InClass, UObject* InParent, FName InName, EObjectFlags Flags, UObject* Context, const TCHAR* Type, const uint8*& Buffer, const uint8* BufferEnd, FFeedbackContext* Warn)
+
+void ProcessTriangleCorner( SpeedTree::CCore& SpeedTree, const int32 TriangleIndex, const int32 Corner, const SpeedTree::SDrawCall* DrawCall, const SpeedTree::st_uint32* Indices32, const SpeedTree::st_uint16* Indices16, FRawMesh& RawMesh, const int32 IndexOffset, const int32 NumUVs, const SpeedTree::SRenderState* RenderState )
+{
+	SpeedTree::st_float32 Data[ 4 ];
+
+	// flip the triangle winding
+	int32 Index = TriangleIndex * 3 + Corner;
+
+	int32 VertexIndex = DrawCall->m_b32BitIndices ? Indices32[ Index ] : Indices16[ Index ];
+	RawMesh.WedgeIndices.Add( VertexIndex + IndexOffset );
+
+	// tangents
+	DrawCall->GetProperty( SpeedTree::VERTEX_PROPERTY_NORMAL, VertexIndex, Data );
+	FVector Normal( -Data[ 0 ], Data[ 1 ], Data[ 2 ] );
+	DrawCall->GetProperty( SpeedTree::VERTEX_PROPERTY_TANGENT, VertexIndex, Data );
+	FVector Tangent( -Data[ 0 ], Data[ 1 ], Data[ 2 ] );
+	RawMesh.WedgeTangentX.Add( Tangent );
+	RawMesh.WedgeTangentY.Add( Normal ^ Tangent );
+	RawMesh.WedgeTangentZ.Add( Normal );
+
+	// ao
+	DrawCall->GetProperty( SpeedTree::VERTEX_PROPERTY_AMBIENT_OCCLUSION, VertexIndex, Data );
+	uint8 AO = Data[ 0 ] * 255.0f;
+	RawMesh.WedgeColors.Add( FColor( AO, AO, AO, 255 ) );
+
+	// keep texcoords padded to align indices
+	int32 BaseTexcoordIndex = RawMesh.WedgeTexCoords[ 0 ].Num();
+	for( int32 PadIndex = 0; PadIndex < NumUVs; ++PadIndex )
+	{
+		RawMesh.WedgeTexCoords[ PadIndex ].AddUninitialized( 1 );
+	}
+
+	// All texcoords are packed into 4 float4 vertex attributes
+	// Data is as follows								
+
+	//		Branches			Fronds				Leaves				Billboards
+	//
+	// 0	Diffuse				Diffuse				Diffuse				Diffuse			
+	// 1	Lightmap UV			Lightmap UV			Lightmap UV			Lightmap UV	(same as diffuse)		
+	// 2	Branch Wind XY		Branch Wind XY		Branch Wind XY			
+	// 3	LOD XY				LOD XY				LOD XY				
+	// 4	LOD Z, Seam Amount	LOD Z, 0			LOD Z, Anchor X		
+	// 5	Detail UV			Frond Wind XY		Anchor YZ	
+	// 6	Seam UV				Frond Wind Z, 0		Leaf Wind XY
+	// 7	0					0					Leaf Wind Z, Leaf Group
+
+
+	// diffuse
+	DrawCall->GetProperty( SpeedTree::VERTEX_PROPERTY_DIFFUSE_TEXCOORDS, VertexIndex, Data );
+	RawMesh.WedgeTexCoords[ 0 ].Top() = FVector2D( Data[ 0 ], Data[ 1 ] );
+
+	// lightmap
+	DrawCall->GetProperty( SpeedTree::VERTEX_PROPERTY_LIGHTMAP_TEXCOORDS, VertexIndex, Data );
+	RawMesh.WedgeTexCoords[ 1 ].Top() = FVector2D( Data[ 0 ], Data[ 1 ] );
+
+	// branch wind
+	DrawCall->GetProperty( SpeedTree::VERTEX_PROPERTY_WIND_BRANCH_DATA, VertexIndex, Data );
+	RawMesh.WedgeTexCoords[ 2 ].Top() = FVector2D( Data[ 0 ], Data[ 1 ] );
+
+	// lod
+	if( RenderState->m_bFacingLeavesPresent )
+	{
+		DrawCall->GetProperty( SpeedTree::VERTEX_PROPERTY_LEAF_CARD_LOD_SCALAR, VertexIndex, Data );
+		RawMesh.WedgeTexCoords[ 3 ].Top() = FVector2D( Data[ 0 ], 0.0f );
+		RawMesh.WedgeTexCoords[ 4 ].Top() = FVector2D( 0.0f, 0.0f );
+	}
+	else
+	{
+		DrawCall->GetProperty( SpeedTree::VERTEX_PROPERTY_LOD_POSITION, VertexIndex, Data );
+		RawMesh.WedgeTexCoords[ 3 ].Top() = FVector2D( -Data[ 0 ], Data[ 1 ] );
+		RawMesh.WedgeTexCoords[ 4 ].Top() = FVector2D( Data[ 2 ], 0.0f );
+	}
+
+	// other
+	if( RenderState->m_bBranchesPresent )
+	{
+		// detail
+		DrawCall->GetProperty( SpeedTree::VERTEX_PROPERTY_DETAIL_TEXCOORDS, VertexIndex, Data );
+		RawMesh.WedgeTexCoords[ 5 ].Top() = FVector2D( Data[ 0 ], Data[ 1 ] );
+
+		// branch seam
+		DrawCall->GetProperty( SpeedTree::VERTEX_PROPERTY_BRANCH_SEAM_DIFFUSE, VertexIndex, Data );
+		RawMesh.WedgeTexCoords[ 6 ].Top() = FVector2D( Data[ 0 ], Data[ 1 ] );
+		RawMesh.WedgeTexCoords[ 4 ].Top().Y = Data[ 2 ];
+	}
+	else if( RenderState->m_bFrondsPresent )
+	{
+		// frond wind
+		DrawCall->GetProperty( SpeedTree::VERTEX_PROPERTY_WIND_EXTRA_DATA, VertexIndex, Data );
+		RawMesh.WedgeTexCoords[ 5 ].Top() = FVector2D( Data[ 0 ], Data[ 1 ] );
+		RawMesh.WedgeTexCoords[ 6 ].Top() = FVector2D( Data[ 2 ], 0.0f );
+	}
+	else if( RenderState->m_bLeavesPresent || RenderState->m_bFacingLeavesPresent )
+	{
+		// anchor
+		if( RenderState->m_bFacingLeavesPresent )
+		{
+			DrawCall->GetProperty( SpeedTree::VERTEX_PROPERTY_POSITION, VertexIndex, Data );
+		}
+		else
+		{
+			DrawCall->GetProperty( SpeedTree::VERTEX_PROPERTY_LEAF_ANCHOR_POINT, VertexIndex, Data );
+		}
+		RawMesh.WedgeTexCoords[ 4 ].Top().Y = -Data[ 0 ];
+		RawMesh.WedgeTexCoords[ 5 ].Top() = FVector2D( Data[ 1 ], Data[ 2 ] );
+
+		// leaf wind
+		DrawCall->GetProperty( SpeedTree::VERTEX_PROPERTY_WIND_EXTRA_DATA, VertexIndex, Data );
+		RawMesh.WedgeTexCoords[ 6 ].Top() = FVector2D( Data[ 0 ], Data[ 1 ] );
+		RawMesh.WedgeTexCoords[ 7 ].Top().X = Data[ 2 ];
+		DrawCall->GetProperty( SpeedTree::VERTEX_PROPERTY_WIND_FLAGS, VertexIndex, Data );
+		RawMesh.WedgeTexCoords[ 7 ].Top().Y = Data[ 0 ];
+	}
+}
+
+
+UObject* USpeedTreeImportFactory::FactoryCreateBinary( UClass* InClass, UObject* InParent, FName InName, EObjectFlags Flags, UObject* Context, const TCHAR* Type, const uint8*& Buffer, const uint8* BufferEnd, FFeedbackContext* Warn )
 {
 	FEditorDelegates::OnAssetPreImport.Broadcast(this, InClass, InParent, InName, Type);
 
@@ -1292,8 +1410,8 @@ UObject* USpeedTreeImportFactory::FactoryCreateBinary(UClass* InClass, UObject*
 							}
 
 							const SpeedTree::st_byte* pIndexData = &*DrawCall->m_pIndexData;
-							SpeedTree::st_uint32* Indices32 = (SpeedTree::st_uint32*)pIndexData;
-							SpeedTree::st_uint16* Indices16 = (SpeedTree::st_uint16*)pIndexData;
+							const SpeedTree::st_uint32* Indices32 = (SpeedTree::st_uint32*)pIndexData;
+							const SpeedTree::st_uint16* Indices16 = (SpeedTree::st_uint16*)pIndexData;
 
 							int32 TriangleCount = DrawCall->m_nNumIndices / 3;
 							int32 ExistingTris = RawMesh.WedgeIndices.Num() / 3;
@@ -1303,7 +1421,7 @@ UObject* USpeedTreeImportFactory::FactoryCreateBinary(UClass* InClass, UObject*
 								RawMesh.FaceMaterialIndices.Add(MaterialIndex);
 								RawMesh.FaceSmoothingMasks.Add(0);
 
-								#if 0
+#if 0
 								float LeafLODScalar = 1.0f;
 								if (RenderState->m_bLeavesPresent)
 								{
@@ -1334,120 +1452,11 @@ UObject* USpeedTreeImportFactory::FactoryCreateBinary(UClass* InClass, UObject*
 
 									LeafLODScalar = FMath::Sqrt(LODArea / OrigArea);
 								}
-								#endif
+#endif
 
 								for (int32 Corner = 0; Corner < 3; ++Corner)
 								{
-									SpeedTree::st_float32 Data[4];
-
-									// flip the triangle winding
-									int32 Index = TriangleIndex * 3 + Corner;
-
-									int32 VertexIndex = DrawCall->m_b32BitIndices ? Indices32[Index] : Indices16[Index];
-									RawMesh.WedgeIndices.Add(VertexIndex + IndexOffset);
-
-									// tangents
-									DrawCall->GetProperty(SpeedTree::VERTEX_PROPERTY_NORMAL, VertexIndex, Data);
-									FVector Normal(-Data[0], Data[1], Data[2]);
-									DrawCall->GetProperty(SpeedTree::VERTEX_PROPERTY_TANGENT, VertexIndex, Data);
-									FVector Tangent(-Data[0], Data[1], Data[2]);
-									RawMesh.WedgeTangentX.Add(Tangent);
-									RawMesh.WedgeTangentY.Add(Normal ^ Tangent);
-									RawMesh.WedgeTangentZ.Add(Normal);
-
-									// ao
-									DrawCall->GetProperty(SpeedTree::VERTEX_PROPERTY_AMBIENT_OCCLUSION, VertexIndex, Data);
-									uint8 AO = Data[0] * 255.0f;
-									RawMesh.WedgeColors.Add(FColor(AO, AO, AO, 255));
-
-									// keep texcoords padded to align indices
-									int32 BaseTexcoordIndex = RawMesh.WedgeTexCoords[0].Num();
-									for (int32 PadIndex = 0; PadIndex < NumUVs; ++PadIndex)
-									{
-										RawMesh.WedgeTexCoords[PadIndex].AddUninitialized(1);
-									}
-									
-									// All texcoords are packed into 4 float4 vertex attributes
-									// Data is as follows								
-
-									//		Branches			Fronds				Leaves				Billboards
-									//
-									// 0	Diffuse				Diffuse				Diffuse				Diffuse			
-									// 1	Lightmap UV			Lightmap UV			Lightmap UV			Lightmap UV	(same as diffuse)		
-									// 2	Branch Wind XY		Branch Wind XY		Branch Wind XY			
-									// 3	LOD XY				LOD XY				LOD XY				
-									// 4	LOD Z, Seam Amount	LOD Z, 0			LOD Z, Anchor X		
-									// 5	Detail UV			Frond Wind XY		Anchor YZ	
-									// 6	Seam UV				Frond Wind Z, 0		Leaf Wind XY
-									// 7	0					0					Leaf Wind Z, Leaf Group
-
-
-									// diffuse
-									DrawCall->GetProperty(SpeedTree::VERTEX_PROPERTY_DIFFUSE_TEXCOORDS, VertexIndex, Data);
-									RawMesh.WedgeTexCoords[0].Top() = FVector2D(Data[0], Data[1]);
-
-									// lightmap
-									DrawCall->GetProperty(SpeedTree::VERTEX_PROPERTY_LIGHTMAP_TEXCOORDS, VertexIndex, Data);
-									RawMesh.WedgeTexCoords[1].Top() = FVector2D(Data[0], Data[1]);
-
-									// branch wind
-									DrawCall->GetProperty(SpeedTree::VERTEX_PROPERTY_WIND_BRANCH_DATA, VertexIndex, Data);
-									RawMesh.WedgeTexCoords[2].Top() = FVector2D(Data[0], Data[1]);
-
-									// lod
-									if (RenderState->m_bFacingLeavesPresent)
-									{
-										DrawCall->GetProperty(SpeedTree::VERTEX_PROPERTY_LEAF_CARD_LOD_SCALAR, VertexIndex, Data);
-										RawMesh.WedgeTexCoords[3].Top() = FVector2D(Data[0], 0.0f);
-										RawMesh.WedgeTexCoords[4].Top() = FVector2D(0.0f, 0.0f);
-									}
-									else
-									{
-										DrawCall->GetProperty(SpeedTree::VERTEX_PROPERTY_LOD_POSITION, VertexIndex, Data);
-										RawMesh.WedgeTexCoords[3].Top() = FVector2D(-Data[0], Data[1]);
-										RawMesh.WedgeTexCoords[4].Top() = FVector2D(Data[2], 0.0f);
-									}
-
-									// other
-									if (RenderState->m_bBranchesPresent)
-									{
-										// detail
-										DrawCall->GetProperty(SpeedTree::VERTEX_PROPERTY_DETAIL_TEXCOORDS, VertexIndex, Data);
-										RawMesh.WedgeTexCoords[5].Top() = FVector2D(Data[0], Data[1]);
-
-										// branch seam
-										DrawCall->GetProperty(SpeedTree::VERTEX_PROPERTY_BRANCH_SEAM_DIFFUSE, VertexIndex, Data);
-										RawMesh.WedgeTexCoords[6].Top() = FVector2D(Data[0], Data[1]);
-										RawMesh.WedgeTexCoords[4].Top().Y = Data[2];
-									}
-									else if (RenderState->m_bFrondsPresent)
-									{
-										// frond wind
-										DrawCall->GetProperty(SpeedTree::VERTEX_PROPERTY_WIND_EXTRA_DATA, VertexIndex, Data);
-										RawMesh.WedgeTexCoords[5].Top() = FVector2D(Data[0], Data[1]);
-										RawMesh.WedgeTexCoords[6].Top() = FVector2D(Data[2], 0.0f);
-									}
-									else if (RenderState->m_bLeavesPresent || RenderState->m_bFacingLeavesPresent)
-									{
-										// anchor
-										if (RenderState->m_bFacingLeavesPresent)
-										{
-											DrawCall->GetProperty(SpeedTree::VERTEX_PROPERTY_POSITION, VertexIndex, Data);
-										}
-										else
-										{
-											DrawCall->GetProperty(SpeedTree::VERTEX_PROPERTY_LEAF_ANCHOR_POINT, VertexIndex, Data);
-										}
-										RawMesh.WedgeTexCoords[4].Top().Y = -Data[0];
-										RawMesh.WedgeTexCoords[5].Top() = FVector2D(Data[1], Data[2]);
-
-										// leaf wind
-										DrawCall->GetProperty(SpeedTree::VERTEX_PROPERTY_WIND_EXTRA_DATA, VertexIndex, Data);
-										RawMesh.WedgeTexCoords[6].Top() = FVector2D(Data[0], Data[1]);
-										RawMesh.WedgeTexCoords[7].Top().X = Data[2];
-										DrawCall->GetProperty(SpeedTree::VERTEX_PROPERTY_WIND_FLAGS, VertexIndex, Data);
-										RawMesh.WedgeTexCoords[7].Top().Y = Data[0];
-									}									
+									ProcessTriangleCorner( SpeedTree, TriangleIndex, Corner, DrawCall, Indices32, Indices16, RawMesh, IndexOffset, NumUVs, RenderState );
 								}
 							}							
 						}
