Trim down the amount of UObject constructors.

diff --git a/Engine/Plugins/2D/Paper2D/Source/Paper2DEditor/Private/Atlasing/PaperAtlasGenerator.cpp b/Engine/Plugins/2D/Paper2D/Source/Paper2DEditor/Private/Atlasing/PaperAtlasGenerator.cpp
--- a/Engine/Plugins/2D/Paper2D/Source/Paper2DEditor/Private/Atlasing/PaperAtlasGenerator.cpp
+++ b/Engine/Plugins/2D/Paper2D/Source/Paper2DEditor/Private/Atlasing/PaperAtlasGenerator.cpp
@@ -168,7 +168,7 @@ void FPaperAtlasGenerator::HandleAssetChangedEvent(UPaperSpriteAtlas* Atlas)
 			else
 			{
 				// The texture never existed - all sprites referencing this MUST be dirty and MUST be updated
-				RemappedAtlasTextures.Add(NewNamedObject<UTexture2D>(Atlas, NAME_None, RF_Public));
+				RemappedAtlasTextures.Add(NewObject<UTexture2D>(Atlas, NAME_None, RF_Public));
 				RemappedAtlasForceDirty.Add(true);
 			}
 		}
diff --git a/Engine/Plugins/2D/Paper2D/Source/Paper2DEditor/Private/Atlasing/PaperSpriteAtlasFactory.cpp b/Engine/Plugins/2D/Paper2D/Source/Paper2DEditor/Private/Atlasing/PaperSpriteAtlasFactory.cpp
--- a/Engine/Plugins/2D/Paper2D/Source/Paper2DEditor/Private/Atlasing/PaperSpriteAtlasFactory.cpp
+++ b/Engine/Plugins/2D/Paper2D/Source/Paper2DEditor/Private/Atlasing/PaperSpriteAtlasFactory.cpp
@@ -18,7 +18,7 @@ UPaperSpriteAtlasFactory::UPaperSpriteAtlasFactory(const FObjectInitializer& Obj
 
 UObject* UPaperSpriteAtlasFactory::FactoryCreateNew(UClass* InClass, UObject* InParent, FName InName, EObjectFlags Flags, UObject* Context, FFeedbackContext* Warn)
 {
-	return ConstructObject<UObject>(SupportedClass, InParent, InName, Flags | RF_Transactional);
+	return NewObject<UObject>(InParent, SupportedClass, InName, Flags | RF_Transactional);
 }
 
 bool UPaperSpriteAtlasFactory::CanCreateNew() const
diff --git a/Engine/Plugins/2D/Paper2D/Source/Paper2DEditor/Private/ContentBrowserExtensions/ContentBrowserExtensions.cpp b/Engine/Plugins/2D/Paper2D/Source/Paper2DEditor/Private/ContentBrowserExtensions/ContentBrowserExtensions.cpp
--- a/Engine/Plugins/2D/Paper2D/Source/Paper2DEditor/Private/ContentBrowserExtensions/ContentBrowserExtensions.cpp
+++ b/Engine/Plugins/2D/Paper2D/Source/Paper2DEditor/Private/ContentBrowserExtensions/ContentBrowserExtensions.cpp
@@ -57,7 +57,7 @@ struct FCreateSpriteFromTextureExtension : public FContentBrowserSelectedAssetEx
 			UTexture2D* Texture = *TextureIt;
 
 			// Create the factory used to generate the sprite
-			UPaperSpriteFactory* SpriteFactory = ConstructObject<UPaperSpriteFactory>(UPaperSpriteFactory::StaticClass());
+			UPaperSpriteFactory* SpriteFactory = NewObject<UPaperSpriteFactory>();
 			SpriteFactory->InitialTexture = Texture;
 
 			// Create the sprite
@@ -141,7 +141,7 @@ struct FCreateTileSetFromTextureExtension : public FContentBrowserSelectedAssetE
 			UTexture2D* Texture = *TextureIt;
 
 			// Create the factory used to generate the tile set
-			UPaperTileSetFactory* TileSetFactory = ConstructObject<UPaperTileSetFactory>(UPaperTileSetFactory::StaticClass());
+			UPaperTileSetFactory* TileSetFactory = NewObject<UPaperTileSetFactory>();
 			TileSetFactory->InitialTexture = Texture;
 
 			// Get a unique name for the tile set
diff --git a/Engine/Plugins/2D/Paper2D/Source/Paper2DEditor/Private/PaperFlipbookFactory.cpp b/Engine/Plugins/2D/Paper2D/Source/Paper2DEditor/Private/PaperFlipbookFactory.cpp
--- a/Engine/Plugins/2D/Paper2D/Source/Paper2DEditor/Private/PaperFlipbookFactory.cpp
+++ b/Engine/Plugins/2D/Paper2D/Source/Paper2DEditor/Private/PaperFlipbookFactory.cpp
@@ -17,7 +17,7 @@ UPaperFlipbookFactory::UPaperFlipbookFactory(const FObjectInitializer& ObjectIni
 
 UObject* UPaperFlipbookFactory::FactoryCreateNew(UClass* Class, UObject* InParent, FName Name, EObjectFlags Flags, UObject* Context, FFeedbackContext* Warn)
 {
-	UPaperFlipbook* NewFlipbook = ConstructObject<UPaperFlipbook>(Class, InParent, Name, Flags | RF_Transactional);
+	UPaperFlipbook* NewFlipbook = NewObject<UPaperFlipbook>(InParent, Class, Name, Flags | RF_Transactional);
 	{
 		FScopedFlipbookMutator EditLock(NewFlipbook);
 		EditLock.KeyFrames = KeyFrames;
diff --git a/Engine/Plugins/2D/Paper2D/Source/Paper2DEditor/Private/PaperSpriteFactory.cpp b/Engine/Plugins/2D/Paper2D/Source/Paper2DEditor/Private/PaperSpriteFactory.cpp
--- a/Engine/Plugins/2D/Paper2D/Source/Paper2DEditor/Private/PaperSpriteFactory.cpp
+++ b/Engine/Plugins/2D/Paper2D/Source/Paper2DEditor/Private/PaperSpriteFactory.cpp
@@ -26,7 +26,7 @@ bool UPaperSpriteFactory::ConfigureProperties()
 
 UObject* UPaperSpriteFactory::FactoryCreateNew(UClass* Class, UObject* InParent, FName Name, EObjectFlags Flags, UObject* Context, FFeedbackContext* Warn)
 {
-	UPaperSprite* NewSprite = ConstructObject<UPaperSprite>(Class, InParent, Name, Flags | RF_Transactional);
+	UPaperSprite* NewSprite = NewObject<UPaperSprite>(InParent, Class, Name, Flags | RF_Transactional);
 
 	FSpriteAssetInitParameters SpriteInitParams;
 	SpriteInitParams.bNewlyCreated = true;
diff --git a/Engine/Plugins/2D/Paper2D/Source/Paper2DEditor/Private/SpriteAssetTypeActions.cpp b/Engine/Plugins/2D/Paper2D/Source/Paper2DEditor/Private/SpriteAssetTypeActions.cpp
--- a/Engine/Plugins/2D/Paper2D/Source/Paper2DEditor/Private/SpriteAssetTypeActions.cpp
+++ b/Engine/Plugins/2D/Paper2D/Source/Paper2DEditor/Private/SpriteAssetTypeActions.cpp
@@ -108,7 +108,7 @@ void FSpriteAssetTypeActions::ExecuteCreateFlipbook(TArray<TWeakObjectPtr<UPaper
 			FString AssetName;
 			FString PackageName;
 
-			UPaperFlipbookFactory* FlipbookFactory = ConstructObject<UPaperFlipbookFactory>(UPaperFlipbookFactory::StaticClass());
+			UPaperFlipbookFactory* FlipbookFactory = NewObject<UPaperFlipbookFactory>();
 			for (int32 SpriteIndex = 0; SpriteIndex < Sprites.Num(); ++SpriteIndex)
 			{
 				UPaperSprite* Sprite = Sprites[SpriteIndex];
diff --git a/Engine/Plugins/2D/Paper2D/Source/Paper2DEditor/Private/SpriteEditor/SpriteEditorViewportClient.cpp b/Engine/Plugins/2D/Paper2D/Source/Paper2DEditor/Private/SpriteEditor/SpriteEditorViewportClient.cpp
--- a/Engine/Plugins/2D/Paper2D/Source/Paper2DEditor/Private/SpriteEditor/SpriteEditorViewportClient.cpp
+++ b/Engine/Plugins/2D/Paper2D/Source/Paper2DEditor/Private/SpriteEditor/SpriteEditorViewportClient.cpp
@@ -1143,7 +1143,7 @@ UPaperSprite* FSpriteEditorViewportClient::CreateNewSprite(FVector2D TopLeft, FV
 	UPaperSprite* CreatedSprite = nullptr;
 
 	// Create the factory used to generate the sprite
-	UPaperSpriteFactory* SpriteFactory = ConstructObject<UPaperSpriteFactory>(UPaperSpriteFactory::StaticClass());
+	UPaperSpriteFactory* SpriteFactory = NewObject<UPaperSpriteFactory>();
 	SpriteFactory->InitialTexture = CurrentSprite->SourceTexture;
 	SpriteFactory->bUseSourceRegion = true;
 	SpriteFactory->InitialSourceUV = TopLeft;
diff --git a/Engine/Plugins/2D/Paper2D/Source/Paper2DEditor/Private/TileMapEditing/TileMapAssetImportData.cpp b/Engine/Plugins/2D/Paper2D/Source/Paper2DEditor/Private/TileMapEditing/TileMapAssetImportData.cpp
--- a/Engine/Plugins/2D/Paper2D/Source/Paper2DEditor/Private/TileMapEditing/TileMapAssetImportData.cpp
+++ b/Engine/Plugins/2D/Paper2D/Source/Paper2DEditor/Private/TileMapEditing/TileMapAssetImportData.cpp
@@ -18,7 +18,7 @@ UTileMapAssetImportData* UTileMapAssetImportData::GetImportDataForTileMap(UPaper
 	UTileMapAssetImportData* ImportData = Cast<UTileMapAssetImportData>(TileMap->AssetImportData);
 	if (ImportData == nullptr)
 	{
-		ImportData = ConstructObject<UTileMapAssetImportData>(UTileMapAssetImportData::StaticClass(), TileMap, NAME_None, RF_NoFlags/*, TemplateForCreation*/);
+		ImportData = NewObject<UTileMapAssetImportData>(TileMap, NAME_None, RF_NoFlags/*, TemplateForCreation*/);
 
 		// Try to preserve the source file path if possible
 		if (TileMap->AssetImportData != nullptr)
diff --git a/Engine/Plugins/2D/Paper2D/Source/PaperJsonImporter/Private/PaperJsonImporterFactory.cpp b/Engine/Plugins/2D/Paper2D/Source/PaperJsonImporter/Private/PaperJsonImporterFactory.cpp
--- a/Engine/Plugins/2D/Paper2D/Source/PaperJsonImporter/Private/PaperJsonImporterFactory.cpp
+++ b/Engine/Plugins/2D/Paper2D/Source/PaperJsonImporter/Private/PaperJsonImporterFactory.cpp
@@ -67,12 +67,12 @@ UObject* UPaperJsonImporterFactory::FactoryCreateText(UClass* InClass, UObject*
 	if (Importer.ImportFromString(FileContent, NameForErrors) &&
 		Importer.ImportTextures(LongPackagePath, CurrentSourcePath))
 	{
-		UPaperSpriteSheet* SpriteSheet = NewNamedObject<UPaperSpriteSheet>(InParent, InName, Flags);
+		UPaperSpriteSheet* SpriteSheet = NewObject<UPaperSpriteSheet>(InParent, InName, Flags);
 		if (Importer.PerformImport(LongPackagePath, Flags, SpriteSheet))
 		{
 			if (SpriteSheet->AssetImportData == nullptr)
 			{
-				SpriteSheet->AssetImportData = ConstructObject<UAssetImportData>(UAssetImportData::StaticClass(), SpriteSheet);
+				SpriteSheet->AssetImportData = NewObject<UAssetImportData>(SpriteSheet);
 			}
 
 			SpriteSheet->AssetImportData->SourceFilePath = FReimportManager::SanitizeImportFilename(CurrentFilename, SpriteSheet);
diff --git a/Engine/Plugins/2D/Paper2D/Source/PaperJsonImporter/Private/PaperJsonSpriteSheetImporter.cpp b/Engine/Plugins/2D/Paper2D/Source/PaperJsonImporter/Private/PaperJsonSpriteSheetImporter.cpp
--- a/Engine/Plugins/2D/Paper2D/Source/PaperJsonImporter/Private/PaperJsonSpriteSheetImporter.cpp
+++ b/Engine/Plugins/2D/Paper2D/Source/PaperJsonImporter/Private/PaperJsonSpriteSheetImporter.cpp
@@ -497,7 +497,7 @@ bool FPaperJsonSpriteSheetImporter::PerformImport(const FString& LongPackagePath
 		
 		if (TargetSprite == nullptr)
 		{
-			TargetSprite = NewNamedObject<UPaperSprite>(OuterForFrame, *AssetName, Flags);
+			TargetSprite = NewObject<UPaperSprite>(OuterForFrame, *AssetName, Flags);
 			FAssetRegistryModule::AssetCreated(TargetSprite);
 		}
 
diff --git a/Engine/Plugins/2D/Paper2D/Source/PaperJsonImporter/Private/PaperSpriteSheetAssetTypeActions.cpp b/Engine/Plugins/2D/Paper2D/Source/PaperJsonImporter/Private/PaperSpriteSheetAssetTypeActions.cpp
--- a/Engine/Plugins/2D/Paper2D/Source/PaperJsonImporter/Private/PaperSpriteSheetAssetTypeActions.cpp
+++ b/Engine/Plugins/2D/Paper2D/Source/PaperJsonImporter/Private/PaperSpriteSheetAssetTypeActions.cpp
@@ -112,7 +112,7 @@ void FPaperSpriteSheetAssetTypeActions::ExecuteCreateFlipbooks(TArray<TWeakObjec
 			// Create one flipbook for every grouped flipbook name
 			if (SpriteFlipbookMap.Num() > 0)
 			{
-				UPaperFlipbookFactory* FlipbookFactory = ConstructObject<UPaperFlipbookFactory>(UPaperFlipbookFactory::StaticClass());
+				UPaperFlipbookFactory* FlipbookFactory = NewObject<UPaperFlipbookFactory>();
 
 				GWarn->BeginSlowTask(NSLOCTEXT("Paper2D", "Paper2D_CreateFlipbooks", "Creating flipbooks from selection"), true, true);
 
diff --git a/Engine/Plugins/2D/Paper2D/Source/PaperTiledImporter/Private/PaperTiledImporterFactory.cpp b/Engine/Plugins/2D/Paper2D/Source/PaperTiledImporter/Private/PaperTiledImporterFactory.cpp
--- a/Engine/Plugins/2D/Paper2D/Source/PaperTiledImporter/Private/PaperTiledImporterFactory.cpp
+++ b/Engine/Plugins/2D/Paper2D/Source/PaperTiledImporter/Private/PaperTiledImporterFactory.cpp
@@ -184,7 +184,7 @@ UObject* UPaperTiledImporterFactory::FactoryCreateText(UClass* InClass, UObject*
 
 
 		// Create the new tile map asset and import basic/global data
-		Result = NewNamedObject<UPaperTileMap>(InParent, InName, Flags);
+		Result = NewObject<UPaperTileMap>(InParent, InName, Flags);
 
 		Result->Modify();
 		Result->MapWidth = GlobalInfo.Width;
@@ -391,7 +391,7 @@ UObject* UPaperTiledImporterFactory::CreateNewAsset(UClass* AssetClass, const FS
 	UObject* OuterForAsset = CreatePackage(nullptr, *PackageName);
 
 	// Create a frame in the package
-	UObject* NewAsset = ConstructObject<UObject>(AssetClass, OuterForAsset, *AssetName, Flags);
+	UObject* NewAsset = NewObject<UObject>(OuterForAsset, AssetClass, *AssetName, Flags);
 	FAssetRegistryModule::AssetCreated(NewAsset);
 
 	return NewAsset;
diff --git a/Engine/Plugins/Editor/SpeedTreeImporter/Source/SpeedTreeImporter/Private/SpeedTreeImportFactory.cpp b/Engine/Plugins/Editor/SpeedTreeImporter/Source/SpeedTreeImporter/Private/SpeedTreeImportFactory.cpp
--- a/Engine/Plugins/Editor/SpeedTreeImporter/Source/SpeedTreeImporter/Private/SpeedTreeImportFactory.cpp
+++ b/Engine/Plugins/Editor/SpeedTreeImporter/Source/SpeedTreeImporter/Private/SpeedTreeImportFactory.cpp
@@ -668,11 +668,11 @@ UMaterialInterface* CreateSpeedTreeMaterial(UObject* Parent, FString MaterialFul
 	UMaterialExpressionClamp* BranchSeamAmount = NULL;
 	if (Options->IncludeBranchSeamSmoothing->IsChecked() && RenderState->m_bBranchesPresent && RenderState->m_eBranchSeamSmoothing != SpeedTree::EFFECT_OFF)
 	{
-		UMaterialExpressionTextureCoordinate* SeamTexcoordExpression = ConstructObject<UMaterialExpressionTextureCoordinate>(UMaterialExpressionTextureCoordinate::StaticClass(), UnrealMaterial);
+		UMaterialExpressionTextureCoordinate* SeamTexcoordExpression = NewObject<UMaterialExpressionTextureCoordinate>(UnrealMaterial);
 		SeamTexcoordExpression->CoordinateIndex = 4;
 		UnrealMaterial->Expressions.Add(SeamTexcoordExpression);
 
-		UMaterialExpressionComponentMask* ComponentMaskExpression = ConstructObject<UMaterialExpressionComponentMask>(UMaterialExpressionComponentMask::StaticClass(), UnrealMaterial);
+		UMaterialExpressionComponentMask* ComponentMaskExpression = NewObject<UMaterialExpressionComponentMask>(UnrealMaterial);
 		ComponentMaskExpression->R = 0;
 		ComponentMaskExpression->G = 1;
 		ComponentMaskExpression->B = 0;
@@ -680,12 +680,12 @@ UMaterialInterface* CreateSpeedTreeMaterial(UObject* Parent, FString MaterialFul
 		ComponentMaskExpression->Input.Expression = SeamTexcoordExpression;
 		UnrealMaterial->Expressions.Add(ComponentMaskExpression);
 
-		UMaterialExpressionPower* PowerExpression = ConstructObject<UMaterialExpressionPower>(UMaterialExpressionPower::StaticClass(), UnrealMaterial);
+		UMaterialExpressionPower* PowerExpression = NewObject<UMaterialExpressionPower>(UnrealMaterial);
 		PowerExpression->Base.Expression = ComponentMaskExpression;
 		PowerExpression->ConstExponent = RenderState->m_fBranchSeamWeight;
 		UnrealMaterial->Expressions.Add(PowerExpression);
 
-		BranchSeamAmount = ConstructObject<UMaterialExpressionClamp>(UMaterialExpressionClamp::StaticClass(), UnrealMaterial);
+		BranchSeamAmount = NewObject<UMaterialExpressionClamp>(UnrealMaterial);
 		BranchSeamAmount->Input.Expression = PowerExpression;
 		UnrealMaterial->Expressions.Add(BranchSeamAmount);
 	}
@@ -695,7 +695,7 @@ UMaterialInterface* CreateSpeedTreeMaterial(UObject* Parent, FString MaterialFul
 	if (DiffuseTexture)
 	{
 		// make texture sampler
-		UMaterialExpressionTextureSample* TextureExpression = ConstructObject<UMaterialExpressionTextureSample>(UMaterialExpressionTextureSample::StaticClass(), UnrealMaterial);
+		UMaterialExpressionTextureSample* TextureExpression = NewObject<UMaterialExpressionTextureSample>(UnrealMaterial);
 		TextureExpression->Texture = DiffuseTexture;
 		TextureExpression->SamplerType = SAMPLERTYPE_Color;
 		UnrealMaterial->Expressions.Add(TextureExpression);
@@ -712,17 +712,17 @@ UMaterialInterface* CreateSpeedTreeMaterial(UObject* Parent, FString MaterialFul
 		if (BranchSeamAmount != NULL)
 		{
 			// perform branch seam smoothing
-			UMaterialExpressionTextureCoordinate* SeamTexcoordExpression = ConstructObject<UMaterialExpressionTextureCoordinate>(UMaterialExpressionTextureCoordinate::StaticClass(), UnrealMaterial);
+			UMaterialExpressionTextureCoordinate* SeamTexcoordExpression = NewObject<UMaterialExpressionTextureCoordinate>(UnrealMaterial);
 			SeamTexcoordExpression->CoordinateIndex = 6;
 			UnrealMaterial->Expressions.Add(SeamTexcoordExpression);
 
-			UMaterialExpressionTextureSample* SeamTextureExpression = ConstructObject<UMaterialExpressionTextureSample>(UMaterialExpressionTextureSample::StaticClass(), UnrealMaterial);
+			UMaterialExpressionTextureSample* SeamTextureExpression = NewObject<UMaterialExpressionTextureSample>(UnrealMaterial);
 			SeamTextureExpression->Texture = DiffuseTexture;
 			SeamTextureExpression->SamplerType = SAMPLERTYPE_Color;
 			SeamTextureExpression->Coordinates.Expression = SeamTexcoordExpression;
 			UnrealMaterial->Expressions.Add(SeamTextureExpression);
 
-			UMaterialExpressionLinearInterpolate* InterpolateExpression = ConstructObject<UMaterialExpressionLinearInterpolate>(UMaterialExpressionLinearInterpolate::StaticClass(), UnrealMaterial);
+			UMaterialExpressionLinearInterpolate* InterpolateExpression = NewObject<UMaterialExpressionLinearInterpolate>(UnrealMaterial);
 			InterpolateExpression->A.Expression = SeamTextureExpression;
 			InterpolateExpression->B.Expression = TextureExpression;
 			InterpolateExpression->Alpha.Expression = BranchSeamAmount;
@@ -737,19 +737,19 @@ UMaterialInterface* CreateSpeedTreeMaterial(UObject* Parent, FString MaterialFul
 			if (DetailTexture)
 			{
 				// add/find UVSet
-				UMaterialExpressionTextureCoordinate* DetailTexcoordExpression = ConstructObject<UMaterialExpressionTextureCoordinate>(UMaterialExpressionTextureCoordinate::StaticClass(), UnrealMaterial);
+				UMaterialExpressionTextureCoordinate* DetailTexcoordExpression = NewObject<UMaterialExpressionTextureCoordinate>(UnrealMaterial);
 				DetailTexcoordExpression->CoordinateIndex = 5;
 				UnrealMaterial->Expressions.Add(DetailTexcoordExpression);
 				
 				// make texture sampler
-				UMaterialExpressionTextureSample* DetailTextureExpression = ConstructObject<UMaterialExpressionTextureSample>(UMaterialExpressionTextureSample::StaticClass(), UnrealMaterial);
+				UMaterialExpressionTextureSample* DetailTextureExpression = NewObject<UMaterialExpressionTextureSample>(UnrealMaterial);
 				DetailTextureExpression->Texture = DetailTexture;
 				DetailTextureExpression->SamplerType = SAMPLERTYPE_Color;
 				DetailTextureExpression->Coordinates.Expression = DetailTexcoordExpression;
 				UnrealMaterial->Expressions.Add(DetailTextureExpression);
 
 				// interpolate the detail
-				UMaterialExpressionLinearInterpolate* InterpolateExpression = ConstructObject<UMaterialExpressionLinearInterpolate>(UMaterialExpressionLinearInterpolate::StaticClass(), UnrealMaterial);
+				UMaterialExpressionLinearInterpolate* InterpolateExpression = NewObject<UMaterialExpressionLinearInterpolate>(UnrealMaterial);
 				InterpolateExpression->A.Expression = UnrealMaterial->BaseColor.Expression;
 				InterpolateExpression->B.Expression = DetailTextureExpression;
 				InterpolateExpression->Alpha.Expression = DetailTextureExpression;
@@ -773,7 +773,7 @@ UMaterialInterface* CreateSpeedTreeMaterial(UObject* Parent, FString MaterialFul
 		if (SpecularTexture)
 		{
 			// make texture sampler
-			UMaterialExpressionTextureSample* TextureExpression = ConstructObject<UMaterialExpressionTextureSample>(UMaterialExpressionTextureSample::StaticClass(), UnrealMaterial);
+			UMaterialExpressionTextureSample* TextureExpression = NewObject<UMaterialExpressionTextureSample>(UnrealMaterial);
 			TextureExpression->Texture = SpecularTexture;
 			TextureExpression->SamplerType = SAMPLERTYPE_Color;
 			
@@ -785,7 +785,7 @@ UMaterialInterface* CreateSpeedTreeMaterial(UObject* Parent, FString MaterialFul
 
 	if (!bMadeSpecular)
 	{
-		UMaterialExpressionConstant* ZeroExpression = ConstructObject<UMaterialExpressionConstant>(UMaterialExpressionConstant::StaticClass(), UnrealMaterial);
+		UMaterialExpressionConstant* ZeroExpression = NewObject<UMaterialExpressionConstant>(UnrealMaterial);
 		ZeroExpression->R = 0.0f;
 		UnrealMaterial->Expressions.Add(ZeroExpression);
 		UnrealMaterial->Specular.Expression = ZeroExpression;
@@ -797,7 +797,7 @@ UMaterialInterface* CreateSpeedTreeMaterial(UObject* Parent, FString MaterialFul
 		if (NormalTexture)
 		{
 			// make texture sampler
-			UMaterialExpressionTextureSample* TextureExpression = ConstructObject<UMaterialExpressionTextureSample>(UMaterialExpressionTextureSample::StaticClass(), UnrealMaterial);
+			UMaterialExpressionTextureSample* TextureExpression = NewObject<UMaterialExpressionTextureSample>(UnrealMaterial);
 			TextureExpression->Texture = NormalTexture;
 			TextureExpression->SamplerType = SAMPLERTYPE_Normal;
 			
@@ -807,17 +807,17 @@ UMaterialInterface* CreateSpeedTreeMaterial(UObject* Parent, FString MaterialFul
 			if (BranchSeamAmount != NULL)
 			{
 				// perform branch seam smoothing
-				UMaterialExpressionTextureCoordinate* SeamTexcoordExpression = ConstructObject<UMaterialExpressionTextureCoordinate>(UMaterialExpressionTextureCoordinate::StaticClass(), UnrealMaterial);
+				UMaterialExpressionTextureCoordinate* SeamTexcoordExpression = NewObject<UMaterialExpressionTextureCoordinate>(UnrealMaterial);
 				SeamTexcoordExpression->CoordinateIndex = 6;
 				UnrealMaterial->Expressions.Add(SeamTexcoordExpression);
 
-				UMaterialExpressionTextureSample* SeamTextureExpression = ConstructObject<UMaterialExpressionTextureSample>(UMaterialExpressionTextureSample::StaticClass(), UnrealMaterial);
+				UMaterialExpressionTextureSample* SeamTextureExpression = NewObject<UMaterialExpressionTextureSample>(UnrealMaterial);
 				SeamTextureExpression->Texture = NormalTexture;
 				SeamTextureExpression->SamplerType = SAMPLERTYPE_Normal;
 				SeamTextureExpression->Coordinates.Expression = SeamTexcoordExpression;
 				UnrealMaterial->Expressions.Add(SeamTextureExpression);
 
-				UMaterialExpressionLinearInterpolate* InterpolateExpression = ConstructObject<UMaterialExpressionLinearInterpolate>(UMaterialExpressionLinearInterpolate::StaticClass(), UnrealMaterial);
+				UMaterialExpressionLinearInterpolate* InterpolateExpression = NewObject<UMaterialExpressionLinearInterpolate>(UnrealMaterial);
 				InterpolateExpression->A.Expression = SeamTextureExpression;
 				InterpolateExpression->B.Expression = TextureExpression;
 				InterpolateExpression->Alpha.Expression = BranchSeamAmount;
@@ -830,7 +830,7 @@ UMaterialInterface* CreateSpeedTreeMaterial(UObject* Parent, FString MaterialFul
 
 	if (Options->IncludeVertexProcessingCheck->IsChecked() && !RenderState->m_bRigidMeshesPresent)
 	{
-		UMaterialExpressionSpeedTree* SpeedTreeExpression = ConstructObject<UMaterialExpressionSpeedTree>(UMaterialExpressionSpeedTree::StaticClass(), UnrealMaterial);
+		UMaterialExpressionSpeedTree* SpeedTreeExpression = NewObject<UMaterialExpressionSpeedTree>(UnrealMaterial);
 	
 		SpeedTreeExpression->LODType = (Options->IncludeSmoothLODCheck->IsChecked() ? STLOD_Smooth : STLOD_Pop);
 		SpeedTreeExpression->WindType = WindType;
@@ -856,7 +856,7 @@ UMaterialInterface* CreateSpeedTreeMaterial(UObject* Parent, FString MaterialFul
 	if (Options->IncludeSpeedTreeAO->IsChecked() && 
 		!(RenderState->m_bVertBillboard || RenderState->m_bHorzBillboard))
 	{
-		UMaterialExpressionVertexColor* VertexColor = ConstructObject<UMaterialExpressionVertexColor>(UMaterialExpressionVertexColor::StaticClass(), UnrealMaterial);
+		UMaterialExpressionVertexColor* VertexColor = NewObject<UMaterialExpressionVertexColor>(UnrealMaterial);
 		UnrealMaterial->Expressions.Add(VertexColor);
 		UnrealMaterial->AmbientOcclusion.Expression = VertexColor;
 		UnrealMaterial->AmbientOcclusion.Mask = VertexColor->GetOutputs()[0].Mask;
@@ -869,16 +869,16 @@ UMaterialInterface* CreateSpeedTreeMaterial(UObject* Parent, FString MaterialFul
 	// UE4 flips normals for two-sided materials. SpeedTrees don't need that
 	if (UnrealMaterial->TwoSided)
 	{
-		UMaterialExpressionTwoSidedSign* TwoSidedSign = ConstructObject<UMaterialExpressionTwoSidedSign>(UMaterialExpressionTwoSidedSign::StaticClass(), UnrealMaterial);
+		UMaterialExpressionTwoSidedSign* TwoSidedSign = NewObject<UMaterialExpressionTwoSidedSign>(UnrealMaterial);
 		UnrealMaterial->Expressions.Add(TwoSidedSign);
 
-		UMaterialExpressionMultiply* Multiply = ConstructObject<UMaterialExpressionMultiply>(UMaterialExpressionMultiply::StaticClass(), UnrealMaterial);
+		auto Multiply = NewObject<UMaterialExpressionMultiply>(UnrealMaterial);
 		UnrealMaterial->Expressions.Add(Multiply);
 		Multiply->A.Expression = TwoSidedSign;
 
 		if (UnrealMaterial->Normal.Expression == NULL)
 		{
-			UMaterialExpressionConstant3Vector* VertexNormalExpression = ConstructObject<UMaterialExpressionConstant3Vector>(UMaterialExpressionConstant3Vector::StaticClass(), UnrealMaterial);
+			auto VertexNormalExpression = NewObject<UMaterialExpressionConstant3Vector>(UnrealMaterial);
 			UnrealMaterial->Expressions.Add(VertexNormalExpression);
 			VertexNormalExpression->Constant = FLinearColor(0.0f, 0.0f, 1.0f);
 
@@ -898,7 +898,7 @@ UMaterialInterface* CreateSpeedTreeMaterial(UObject* Parent, FString MaterialFul
 		UMaterialFunction* ColorVariationFunction = LoadObject<UMaterialFunction>(NULL, TEXT("/Engine/Functions/Engine_MaterialFunctions01/SpeedTree/SpeedTreeColorVariation.SpeedTreeColorVariation"), NULL, LOAD_None, NULL);
 		if (ColorVariationFunction)
 		{
-			UMaterialExpressionMaterialFunctionCall* ColorVariation = ConstructObject<UMaterialExpressionMaterialFunctionCall>(UMaterialExpressionMaterialFunctionCall::StaticClass(), UnrealMaterial);
+			UMaterialExpressionMaterialFunctionCall* ColorVariation = NewObject<UMaterialExpressionMaterialFunctionCall>(UnrealMaterial);
 			UnrealMaterial->Expressions.Add(ColorVariation);
 
 			ColorVariation->MaterialFunction = ColorVariationFunction;
@@ -1231,10 +1231,10 @@ UObject* USpeedTreeImportFactory::FactoryCreateBinary( UClass* InClass, UObject*
 					ExistingMesh->PreEditChange(NULL);
 				}
 				
-				StaticMesh = NewNamedObject<UStaticMesh>(Package, FName(*MeshName), Flags | RF_Public);
+				StaticMesh = NewObject<UStaticMesh>(Package, FName(*MeshName), Flags | RF_Public);
 
 				// @todo AssetImportData make a data class for speed tree assets
-				StaticMesh->AssetImportData = ConstructObject<UAssetImportData>(UAssetImportData::StaticClass(), StaticMesh);
+				StaticMesh->AssetImportData = NewObject<UAssetImportData>(StaticMesh);
 				StaticMesh->AssetImportData->SourceFilePath = FReimportManager::SanitizeImportFilename(UFactory::GetCurrentFilename(), StaticMesh);
 				StaticMesh->AssetImportData->SourceFileTimestamp = IFileManager::Get().GetTimeStamp(*UFactory::GetCurrentFilename()).ToString();
 				StaticMesh->AssetImportData->bDirty = false;
diff --git a/Engine/Plugins/Media/MediaPlayerEditor/Source/MediaPlayerEditor/Private/AssetTools/MediaPlayerActions.cpp b/Engine/Plugins/Media/MediaPlayerEditor/Source/MediaPlayerEditor/Private/AssetTools/MediaPlayerActions.cpp
--- a/Engine/Plugins/Media/MediaPlayerEditor/Source/MediaPlayerEditor/Private/AssetTools/MediaPlayerActions.cpp
+++ b/Engine/Plugins/Media/MediaPlayerEditor/Source/MediaPlayerEditor/Private/AssetTools/MediaPlayerActions.cpp
@@ -145,7 +145,7 @@ void FMediaPlayerActions::ExecuteCreateMediaSoundWave( TArray<TWeakObjectPtr<UMe
 			CreateUniqueAssetName(Object->GetOutermost()->GetName(), DefaultSuffix, PackagePath, Name);
 
 			// Create the factory used to generate the asset
-			UMediaSoundWaveFactoryNew* Factory = ConstructObject<UMediaSoundWaveFactoryNew>(UMediaSoundWaveFactoryNew::StaticClass());
+			UMediaSoundWaveFactoryNew* Factory = NewObject<UMediaSoundWaveFactoryNew>();
 			Factory->InitialMediaPlayer = Object;
 
 			ContentBrowserSingleton.CreateNewAsset(Name, FPackageName::GetLongPackagePath(PackagePath), UMediaSoundWave::StaticClass(), Factory);
@@ -166,7 +166,7 @@ void FMediaPlayerActions::ExecuteCreateMediaSoundWave( TArray<TWeakObjectPtr<UMe
 				CreateUniqueAssetName(Object->GetOutermost()->GetName(), DefaultSuffix, PackageName, Name);
 
 				// Create the factory used to generate the asset
-				UMediaSoundWaveFactoryNew* Factory = ConstructObject<UMediaSoundWaveFactoryNew>(UMediaSoundWaveFactoryNew::StaticClass());
+				UMediaSoundWaveFactoryNew* Factory = NewObject<UMediaSoundWaveFactoryNew>();
 				Factory->InitialMediaPlayer = Object;
 
 				FAssetToolsModule& AssetToolsModule = FModuleManager::GetModuleChecked<FAssetToolsModule>("AssetTools");
@@ -204,7 +204,7 @@ void FMediaPlayerActions::ExecuteCreateMediaTexture( TArray<TWeakObjectPtr<UMedi
 			CreateUniqueAssetName(Object->GetOutermost()->GetName(), DefaultSuffix, PackagePath, Name);
 
 			// Create the factory used to generate the asset
-			UMediaTextureFactoryNew* Factory = ConstructObject<UMediaTextureFactoryNew>(UMediaTextureFactoryNew::StaticClass());
+			UMediaTextureFactoryNew* Factory = NewObject<UMediaTextureFactoryNew>();
 			Factory->InitialMediaPlayer = Object;
 
 			ContentBrowserSingleton.CreateNewAsset(Name, FPackageName::GetLongPackagePath(PackagePath), UMediaTexture::StaticClass(), Factory);
@@ -225,7 +225,7 @@ void FMediaPlayerActions::ExecuteCreateMediaTexture( TArray<TWeakObjectPtr<UMedi
 				CreateUniqueAssetName(Object->GetOutermost()->GetName(), DefaultSuffix, PackageName, Name);
 
 				// Create the factory used to generate the asset
-				UMediaTextureFactoryNew* Factory = ConstructObject<UMediaTextureFactoryNew>(UMediaTextureFactoryNew::StaticClass());
+				UMediaTextureFactoryNew* Factory = NewObject<UMediaTextureFactoryNew>();
 				Factory->InitialMediaPlayer = Object;
 
 				FAssetToolsModule& AssetToolsModule = FModuleManager::GetModuleChecked<FAssetToolsModule>("AssetTools");
diff --git a/Engine/Plugins/Media/MediaPlayerEditor/Source/MediaPlayerEditor/Private/AssetTools/MediaTextureActions.cpp b/Engine/Plugins/Media/MediaPlayerEditor/Source/MediaPlayerEditor/Private/AssetTools/MediaTextureActions.cpp
--- a/Engine/Plugins/Media/MediaPlayerEditor/Source/MediaPlayerEditor/Private/AssetTools/MediaTextureActions.cpp
+++ b/Engine/Plugins/Media/MediaPlayerEditor/Source/MediaPlayerEditor/Private/AssetTools/MediaTextureActions.cpp
@@ -126,7 +126,7 @@ void FMediaTextureActions::ExecuteCreateMaterial( TArray<TWeakObjectPtr<UTexture
 			CreateUniqueAssetName(Object->GetOutermost()->GetName(), DefaultSuffix, PackagePath, Name);
 
 			// Create the factory used to generate the asset
-			UMaterialFactoryNew* Factory = ConstructObject<UMaterialFactoryNew>(UMaterialFactoryNew::StaticClass());
+			UMaterialFactoryNew* Factory = NewObject<UMaterialFactoryNew>();
 			Factory->InitialTexture = Object;
 
 			ContentBrowserSingleton.CreateNewAsset(Name, FPackageName::GetLongPackagePath(PackagePath), UMaterial::StaticClass(), Factory);
@@ -147,7 +147,7 @@ void FMediaTextureActions::ExecuteCreateMaterial( TArray<TWeakObjectPtr<UTexture
 				CreateUniqueAssetName(Object->GetOutermost()->GetName(), DefaultSuffix, PackageName, Name);
 
 				// Create the factory used to generate the asset
-				UMaterialFactoryNew* Factory = ConstructObject<UMaterialFactoryNew>(UMaterialFactoryNew::StaticClass());
+				UMaterialFactoryNew* Factory = NewObject<UMaterialFactoryNew>();
 				Factory->InitialTexture = Object;
 
 				FAssetToolsModule& AssetToolsModule = FModuleManager::GetModuleChecked<FAssetToolsModule>("AssetTools");
diff --git a/Engine/Plugins/Media/MediaPlayerEditor/Source/MediaPlayerEditor/Private/Factories/MediaSoundWaveFactoryNew.cpp b/Engine/Plugins/Media/MediaPlayerEditor/Source/MediaPlayerEditor/Private/Factories/MediaSoundWaveFactoryNew.cpp
--- a/Engine/Plugins/Media/MediaPlayerEditor/Source/MediaPlayerEditor/Private/Factories/MediaSoundWaveFactoryNew.cpp
+++ b/Engine/Plugins/Media/MediaPlayerEditor/Source/MediaPlayerEditor/Private/Factories/MediaSoundWaveFactoryNew.cpp
@@ -20,7 +20,7 @@ UMediaSoundWaveFactoryNew::UMediaSoundWaveFactoryNew( const FObjectInitializer&
 
 UObject* UMediaSoundWaveFactoryNew::FactoryCreateNew( UClass* InClass, UObject* InParent, FName InName, EObjectFlags Flags, UObject* Context, FFeedbackContext* Warn )
 {
-	UMediaSoundWave* MediaSoundWave = ConstructObject<UMediaSoundWave>(InClass, InParent, InName, Flags);
+	UMediaSoundWave* MediaSoundWave = NewObject<UMediaSoundWave>(InParent, InClass, InName, Flags);
 
 	if ((MediaSoundWave != nullptr) && (InitialMediaPlayer != nullptr))
 	{
diff --git a/Engine/Plugins/Media/MediaPlayerEditor/Source/MediaPlayerEditor/Private/Factories/MediaTextureFactoryNew.cpp b/Engine/Plugins/Media/MediaPlayerEditor/Source/MediaPlayerEditor/Private/Factories/MediaTextureFactoryNew.cpp
--- a/Engine/Plugins/Media/MediaPlayerEditor/Source/MediaPlayerEditor/Private/Factories/MediaTextureFactoryNew.cpp
+++ b/Engine/Plugins/Media/MediaPlayerEditor/Source/MediaPlayerEditor/Private/Factories/MediaTextureFactoryNew.cpp
@@ -20,7 +20,7 @@ UMediaTextureFactoryNew::UMediaTextureFactoryNew( const FObjectInitializer& Obje
 
 UObject* UMediaTextureFactoryNew::FactoryCreateNew( UClass* InClass, UObject* InParent, FName InName, EObjectFlags Flags, UObject* Context, FFeedbackContext* Warn )
 {
-	UMediaTexture* MediaTexture = ConstructObject<UMediaTexture>(InClass, InParent, InName, Flags);
+	UMediaTexture* MediaTexture = NewObject<UMediaTexture>(InParent, InClass, InName, Flags);
 
 	if ((MediaTexture != nullptr) && (InitialMediaPlayer != nullptr))
 	{
diff --git a/Engine/Plugins/ScriptGeneratorPlugin/Source/ScriptGeneratorPlugin/Private/GenericScriptCodeGenerator.cpp b/Engine/Plugins/ScriptGeneratorPlugin/Source/ScriptGeneratorPlugin/Private/GenericScriptCodeGenerator.cpp
--- a/Engine/Plugins/ScriptGeneratorPlugin/Source/ScriptGeneratorPlugin/Private/GenericScriptCodeGenerator.cpp
+++ b/Engine/Plugins/ScriptGeneratorPlugin/Source/ScriptGeneratorPlugin/Private/GenericScriptCodeGenerator.cpp
@@ -157,7 +157,7 @@ FString FGenericScriptCodeGenerator::ExportAdditionalClassGlue(const FString& Cl
 		GeneratedGlue += TEXT("\r\n{\r\n");
 		GeneratedGlue += TEXT("\tUObject* Outer = NULL;\r\n");
 		GeneratedGlue += TEXT("\tFName Name = FName(\"ScriptObject\");\r\n");
-		GeneratedGlue += FString::Printf(TEXT("\tUObject* Obj = NewNamedObject<%s>(Outer, Name);\r\n"), *ClassNameCPP);
+		GeneratedGlue += FString::Printf(TEXT("\tUObject* Obj = NewObject<%s>(Outer, Name);\r\n"), *ClassNameCPP);
 		GeneratedGlue += TEXT("\tif (Obj)\r\n\t{\r\n");
 		GeneratedGlue += TEXT("\t\tFScriptObjectReferencer::Get().AddObjectReference(Obj);\r\n");
 		GeneratedGlue += TEXT("\t\t// @todo: Register the object with the script context here\r\n");
diff --git a/Engine/Plugins/ScriptPlugin/Source/ScriptEditorPlugin/Private/ScriptBlueprintCompiler.cpp b/Engine/Plugins/ScriptPlugin/Source/ScriptEditorPlugin/Private/ScriptBlueprintCompiler.cpp
--- a/Engine/Plugins/ScriptPlugin/Source/ScriptEditorPlugin/Private/ScriptBlueprintCompiler.cpp
+++ b/Engine/Plugins/ScriptPlugin/Source/ScriptEditorPlugin/Private/ScriptBlueprintCompiler.cpp
@@ -132,7 +132,7 @@ void FScriptBlueprintCompiler::CreateScriptDefinedFunction(FScriptField& Field)
 	// Create Blueprint Graph which consists of 3 nodes: 'Entry', 'Get Script Context' and 'Call Function'
 	// @todo: once we figure out how to get parameter lists for functions we can add suport for that here
 
-	UEdGraph* ScriptFunctionGraph = NewNamedObject<UEdGraph>(Blueprint, *FString::Printf(TEXT("%s_Graph"), *FunctionName));
+	UEdGraph* ScriptFunctionGraph = NewObject<UEdGraph>(Blueprint, *FString::Printf(TEXT("%s_Graph"), *FunctionName));
 	ScriptFunctionGraph->Schema = UEdGraphSchema_K2::StaticClass();
 	ScriptFunctionGraph->SetFlags(RF_Transient);
 	
@@ -219,7 +219,7 @@ void FScriptBlueprintCompiler::SpawnNewClass(const FString& NewClassName)
 
 	if ( NewScriptBlueprintClass == NULL )
 	{
-		NewScriptBlueprintClass = ConstructObject<UScriptBlueprintGeneratedClass>(UScriptBlueprintGeneratedClass::StaticClass(), Blueprint->GetOutermost(), FName(*NewClassName), RF_Public | RF_Transactional);
+		NewScriptBlueprintClass = NewObject<UScriptBlueprintGeneratedClass>(Blueprint->GetOutermost(), FName(*NewClassName), RF_Public | RF_Transactional);
 	}
 	else
 	{
diff --git a/Engine/Source/Developer/AssetTools/Private/AssetTools.cpp b/Engine/Source/Developer/AssetTools/Private/AssetTools.cpp
--- a/Engine/Source/Developer/AssetTools/Private/AssetTools.cpp
+++ b/Engine/Source/Developer/AssetTools/Private/AssetTools.cpp
@@ -570,7 +570,7 @@ TArray<UObject*> FAssetTools::ImportAssets(const TArray<FString>& Files, const F
 								// We found a factory for this file, it can be imported!
 								// Create a new factory of the same class and make sure it doesn't get GCed.
 								// The object will be removed from the root set at the end of this function.
-								UFactory* NewFactory = ConstructObject<UFactory>( Factory->GetClass() );
+								UFactory* NewFactory = NewObject<UFactory>(GetTransientPackage(), Factory->GetClass());
 								if ( NewFactory->ConfigureProperties() )
 								{
 									NewFactory->AddToRoot();
@@ -615,7 +615,7 @@ TArray<UObject*> FAssetTools::ImportAssets(const TArray<FString>& Files, const F
 						// We found a factory for this file, it can be imported!
 						// Create a new factory of the same class and make sure it doesnt get GCed.
 						// The object will be removed from the root set at the end of this function.
-						UFactory* NewFactory = ConstructObject<UFactory>( SpecifiedFactory->GetClass() );
+						UFactory* NewFactory = NewObject<UFactory>(GetTransientPackage(), SpecifiedFactory->GetClass());
 						if ( NewFactory->ConfigureProperties() )
 						{
 							NewFactory->AddToRoot();
diff --git a/Engine/Source/Developer/AssetTools/Private/AssetTypeActions/AssetTypeActions_AnimBlueprint.cpp b/Engine/Source/Developer/AssetTools/Private/AssetTypeActions/AssetTypeActions_AnimBlueprint.cpp
--- a/Engine/Source/Developer/AssetTools/Private/AssetTypeActions/AssetTypeActions_AnimBlueprint.cpp
+++ b/Engine/Source/Developer/AssetTools/Private/AssetTypeActions/AssetTypeActions_AnimBlueprint.cpp
@@ -78,7 +78,7 @@ UThumbnailInfo* FAssetTypeActions_AnimBlueprint::GetThumbnailInfo(UObject* Asset
 	UThumbnailInfo* ThumbnailInfo = AnimBlueprint->ThumbnailInfo;
 	if (ThumbnailInfo == NULL)
 	{
-		ThumbnailInfo = ConstructObject<USceneThumbnailInfo>(USceneThumbnailInfo::StaticClass(), AnimBlueprint);
+		ThumbnailInfo = NewObject<USceneThumbnailInfo>(AnimBlueprint);
 		AnimBlueprint->ThumbnailInfo = ThumbnailInfo;
 	}
 
diff --git a/Engine/Source/Developer/AssetTools/Private/AssetTypeActions/AssetTypeActions_AnimSequence.cpp b/Engine/Source/Developer/AssetTools/Private/AssetTypeActions/AssetTypeActions_AnimSequence.cpp
--- a/Engine/Source/Developer/AssetTools/Private/AssetTypeActions/AssetTypeActions_AnimSequence.cpp
+++ b/Engine/Source/Developer/AssetTools/Private/AssetTypeActions/AssetTypeActions_AnimSequence.cpp
@@ -89,7 +89,7 @@ void FAssetTypeActions_AnimSequence::ExecuteReimportWithNewSource(TArray<TWeakOb
 void FAssetTypeActions_AnimSequence::ExecuteNewAnimComposite(TArray<TWeakObjectPtr<UAnimSequence>> Objects) const
 {
 	const FString DefaultSuffix = TEXT("_Composite");
-	UAnimCompositeFactory* Factory = ConstructObject<UAnimCompositeFactory>(UAnimCompositeFactory::StaticClass());
+	UAnimCompositeFactory* Factory = NewObject<UAnimCompositeFactory>();
 
 	CreateAnimationAssets(Objects, UAnimComposite::StaticClass(), Factory, DefaultSuffix, FOnConfigureFactory::CreateSP(this, &FAssetTypeActions_AnimSequence::ConfigureFactoryForAnimComposite));
 }
@@ -97,7 +97,7 @@ void FAssetTypeActions_AnimSequence::ExecuteNewAnimComposite(TArray<TWeakObjectP
 void FAssetTypeActions_AnimSequence::ExecuteNewAnimMontage(TArray<TWeakObjectPtr<UAnimSequence>> Objects) const
 {
 	const FString DefaultSuffix = TEXT("_Montage");
-	UAnimMontageFactory* Factory = ConstructObject<UAnimMontageFactory>(UAnimMontageFactory::StaticClass());
+	UAnimMontageFactory* Factory = NewObject<UAnimMontageFactory>();
 
 	CreateAnimationAssets(Objects, UAnimMontage::StaticClass(), Factory, DefaultSuffix, FOnConfigureFactory::CreateSP(this, &FAssetTypeActions_AnimSequence::ConfigureFactoryForAnimMontage));
 }
diff --git a/Engine/Source/Developer/AssetTools/Private/AssetTypeActions/AssetTypeActions_AnimationAsset.cpp b/Engine/Source/Developer/AssetTools/Private/AssetTypeActions/AssetTypeActions_AnimationAsset.cpp
--- a/Engine/Source/Developer/AssetTools/Private/AssetTypeActions/AssetTypeActions_AnimationAsset.cpp
+++ b/Engine/Source/Developer/AssetTools/Private/AssetTypeActions/AssetTypeActions_AnimationAsset.cpp
@@ -76,7 +76,7 @@ UThumbnailInfo* FAssetTypeActions_AnimationAsset::GetThumbnailInfo(UObject* Asse
 	UThumbnailInfo* ThumbnailInfo = Anim->ThumbnailInfo;
 	if (ThumbnailInfo == NULL)
 	{
-		ThumbnailInfo = ConstructObject<USceneThumbnailInfo>(USceneThumbnailInfo::StaticClass(), Anim);
+		ThumbnailInfo = NewObject<USceneThumbnailInfo>(Anim);
 		Anim->ThumbnailInfo = ThumbnailInfo;
 	}
 
diff --git a/Engine/Source/Developer/AssetTools/Private/AssetTypeActions/AssetTypeActions_Blueprint.cpp b/Engine/Source/Developer/AssetTools/Private/AssetTypeActions/AssetTypeActions_Blueprint.cpp
--- a/Engine/Source/Developer/AssetTools/Private/AssetTypeActions/AssetTypeActions_Blueprint.cpp
+++ b/Engine/Source/Developer/AssetTools/Private/AssetTypeActions/AssetTypeActions_Blueprint.cpp
@@ -273,7 +273,7 @@ UThumbnailInfo* FAssetTypeActions_Blueprint::GetThumbnailInfo(UObject* Asset) co
 	UThumbnailInfo* ThumbnailInfo = Blueprint->ThumbnailInfo;
 	if ( ThumbnailInfo == NULL )
 	{
-		ThumbnailInfo = ConstructObject<USceneThumbnailInfo>(USceneThumbnailInfo::StaticClass(), Blueprint);
+		ThumbnailInfo = NewObject<USceneThumbnailInfo>(Blueprint);
 		Blueprint->ThumbnailInfo = ThumbnailInfo;
 	}
 
diff --git a/Engine/Source/Developer/AssetTools/Private/AssetTypeActions/AssetTypeActions_CameraAnim.cpp b/Engine/Source/Developer/AssetTools/Private/AssetTypeActions/AssetTypeActions_CameraAnim.cpp
--- a/Engine/Source/Developer/AssetTools/Private/AssetTypeActions/AssetTypeActions_CameraAnim.cpp
+++ b/Engine/Source/Developer/AssetTools/Private/AssetTypeActions/AssetTypeActions_CameraAnim.cpp
@@ -24,7 +24,7 @@ void FAssetTypeActions_CameraAnim::CreateMatineeActorForCameraAnim(UCameraAnim*
 	ActorSpawnParameters.Name = InCameraAnim->GetFName();
 	PreviewMatineeActor = GEditor->GetEditorWorldContext().World()->SpawnActor<AMatineeActorCameraAnim>(ActorSpawnParameters);
 	check(PreviewMatineeActor.IsValid());
-	UInterpData* NewData = ConstructObject<UInterpData>( UInterpData::StaticClass(), GetTransientPackage(), NAME_None, RF_Transactional);
+	UInterpData* NewData = NewObject<UInterpData>(GetTransientPackage(), NAME_None, RF_Transactional);
 	PreviewMatineeActor.Get()->MatineeData = NewData;
 	PreviewMatineeActor.Get()->CameraAnim = InCameraAnim;
 }
@@ -114,14 +114,14 @@ UInterpGroup* FAssetTypeActions_CameraAnim::CreateInterpGroup(UCameraAnim* InCam
 	{
 		// create InterpGroup so that we can play animation to this pawn
 		check(PreviewMatineeActor.Get()->MatineeData);
-		UInterpGroup* NewGroup = ConstructObject<UInterpGroup>(UInterpGroup::StaticClass(), PreviewMatineeActor.Get()->MatineeData, NAME_None, RF_Transient);
+		UInterpGroup* NewGroup = NewObject<UInterpGroup>(PreviewMatineeActor.Get()->MatineeData, NAME_None, RF_Transient);
 		NewGroup->GroupName = FName(TEXT("Preview Pawn"));
 		NewGroup->EnsureUniqueName();
 
 		PreviewMatineeActor.Get()->MatineeData->InterpGroups.Add(NewGroup);
 
 		// now add group inst
-		UInterpGroupInst* NewGroupInst = ConstructObject<UInterpGroupInst>(UInterpGroupInst::StaticClass(), PreviewMatineeActor.Get(), NAME_None, RF_Transient);
+		UInterpGroupInst* NewGroupInst = NewObject<UInterpGroupInst>(PreviewMatineeActor.Get(), NAME_None, RF_Transient);
 		// Initialise group instance, saving ref to actor it works on.
 		NewGroupInst->InitGroupInst(NewGroup, PreviewInfo.PawnInst);
 		const int32 NewGroupInstIndex = PreviewMatineeActor.Get()->GroupInst.Add(NewGroupInst);
@@ -133,14 +133,14 @@ UInterpGroup* FAssetTypeActions_CameraAnim::CreateInterpGroup(UCameraAnim* InCam
 		int32 AnimTrackIndex = INDEX_NONE;
 		// add anim track but do not use addtotrack function that does too many things 
 		// Construct track and track instance objects.
-		UInterpTrackAnimControl* AnimTrack = ConstructObject<UInterpTrackAnimControl>( UInterpTrackAnimControl::StaticClass(), NewGroup, NAME_None, RF_Transient );
+		UInterpTrackAnimControl* AnimTrack = NewObject<UInterpTrackAnimControl>(NewGroup, NAME_None, RF_Transient);
 		check(AnimTrack);
 
 		NewGroup->InterpTracks.Add(AnimTrack);
 		
 		// use config anim slot
 		AnimTrack->SlotName = FName(*GConfig->GetStr(TEXT("MatineePreview"), TEXT("DefaultAnimSlotName"), GEditorIni));
-		UInterpTrackInst* NewTrackInst = ConstructObject<UInterpTrackInst>( AnimTrack->TrackInstClass, NewGroupInst, NAME_None, RF_Transient );
+		UInterpTrackInst* NewTrackInst = NewObject<UInterpTrackInst>(NewGroupInst, AnimTrack->TrackInstClass, NAME_None, RF_Transient);
 		check(NewTrackInst);
 
 		NewGroupInst->TrackInst.Add(NewTrackInst);
diff --git a/Engine/Source/Developer/AssetTools/Private/AssetTypeActions/AssetTypeActions_MaterialFunction.cpp b/Engine/Source/Developer/AssetTools/Private/AssetTypeActions/AssetTypeActions_MaterialFunction.cpp
--- a/Engine/Source/Developer/AssetTools/Private/AssetTypeActions/AssetTypeActions_MaterialFunction.cpp
+++ b/Engine/Source/Developer/AssetTools/Private/AssetTypeActions/AssetTypeActions_MaterialFunction.cpp
@@ -76,7 +76,7 @@ UThumbnailInfo* FAssetTypeActions_MaterialFunction::GetThumbnailInfo(UObject* As
 		ThumbnailInfo = MaterialFunc->ThumbnailInfo;
 		if ( ThumbnailInfo == NULL )
 		{
-			ThumbnailInfo = ConstructObject<USceneThumbnailInfoWithPrimitive>(USceneThumbnailInfoWithPrimitive::StaticClass(), MaterialFunc);
+			ThumbnailInfo = NewObject<USceneThumbnailInfoWithPrimitive>(MaterialFunc);
 			MaterialFunc->ThumbnailInfo = ThumbnailInfo;
 		}
 	}
diff --git a/Engine/Source/Developer/AssetTools/Private/AssetTypeActions/AssetTypeActions_MaterialInterface.cpp b/Engine/Source/Developer/AssetTools/Private/AssetTypeActions/AssetTypeActions_MaterialInterface.cpp
--- a/Engine/Source/Developer/AssetTools/Private/AssetTypeActions/AssetTypeActions_MaterialInterface.cpp
+++ b/Engine/Source/Developer/AssetTools/Private/AssetTypeActions/AssetTypeActions_MaterialInterface.cpp
@@ -29,7 +29,7 @@ UThumbnailInfo* FAssetTypeActions_MaterialInterface::GetThumbnailInfo(UObject* A
 	UThumbnailInfo* ThumbnailInfo = MaterialInterface->ThumbnailInfo;
 	if ( ThumbnailInfo == NULL )
 	{
-		ThumbnailInfo = ConstructObject<USceneThumbnailInfoWithPrimitive>(USceneThumbnailInfoWithPrimitive::StaticClass(), MaterialInterface);
+		ThumbnailInfo = NewObject<USceneThumbnailInfoWithPrimitive>(MaterialInterface);
 		MaterialInterface->ThumbnailInfo = ThumbnailInfo;
 	}
 
@@ -51,7 +51,7 @@ void FAssetTypeActions_MaterialInterface::ExecuteNewMIC(TArray<TWeakObjectPtr<UM
 			FString PackageName;
 			CreateUniqueAssetName(Object->GetOutermost()->GetName(), DefaultSuffix, PackageName, Name);
 
-			UMaterialInstanceConstantFactoryNew* Factory = ConstructObject<UMaterialInstanceConstantFactoryNew>(UMaterialInstanceConstantFactoryNew::StaticClass());
+			UMaterialInstanceConstantFactoryNew* Factory = NewObject<UMaterialInstanceConstantFactoryNew>();
 			Factory->InitialParent = Object;
 
 			FContentBrowserModule& ContentBrowserModule = FModuleManager::LoadModuleChecked<FContentBrowserModule>("ContentBrowser");
@@ -72,7 +72,7 @@ void FAssetTypeActions_MaterialInterface::ExecuteNewMIC(TArray<TWeakObjectPtr<UM
 				CreateUniqueAssetName(Object->GetOutermost()->GetName(), DefaultSuffix, PackageName, Name);
 
 				// Create the factory used to generate the asset
-				UMaterialInstanceConstantFactoryNew* Factory = ConstructObject<UMaterialInstanceConstantFactoryNew>(UMaterialInstanceConstantFactoryNew::StaticClass());
+				UMaterialInstanceConstantFactoryNew* Factory = NewObject<UMaterialInstanceConstantFactoryNew>();
 				Factory->InitialParent = Object;
 
 				FAssetToolsModule& AssetToolsModule = FModuleManager::GetModuleChecked<FAssetToolsModule>("AssetTools");
diff --git a/Engine/Source/Developer/AssetTools/Private/AssetTypeActions/AssetTypeActions_SkeletalMesh.cpp b/Engine/Source/Developer/AssetTools/Private/AssetTypeActions/AssetTypeActions_SkeletalMesh.cpp
--- a/Engine/Source/Developer/AssetTools/Private/AssetTypeActions/AssetTypeActions_SkeletalMesh.cpp
+++ b/Engine/Source/Developer/AssetTools/Private/AssetTypeActions/AssetTypeActions_SkeletalMesh.cpp
@@ -495,7 +495,7 @@ void FAssetTypeActions_SkeletalMesh::OpenAssetEditor( const TArray<UObject*>& In
 					FString PackageName;
 					CreateUniqueAssetName(Mesh->GetOutermost()->GetName(), DefaultSuffix, PackageName, Name);
 
-					USkeletonFactory* Factory = ConstructObject<USkeletonFactory>(USkeletonFactory::StaticClass());
+					USkeletonFactory* Factory = NewObject<USkeletonFactory>();
 					Factory->TargetSkeletalMesh = Mesh;
 
 					FAssetToolsModule& AssetToolsModule = FModuleManager::GetModuleChecked<FAssetToolsModule>("AssetTools");
@@ -538,7 +538,7 @@ UThumbnailInfo* FAssetTypeActions_SkeletalMesh::GetThumbnailInfo(UObject* Asset)
 	UThumbnailInfo* ThumbnailInfo = SkeletalMesh->ThumbnailInfo;
 	if ( ThumbnailInfo == NULL )
 	{
-		ThumbnailInfo = ConstructObject<USceneThumbnailInfo>(USceneThumbnailInfo::StaticClass(), SkeletalMesh);
+		ThumbnailInfo = NewObject<USceneThumbnailInfo>(SkeletalMesh);
 		SkeletalMesh->ThumbnailInfo = ThumbnailInfo;
 	}
 
@@ -601,7 +601,7 @@ void FAssetTypeActions_SkeletalMesh::ExecuteNewSkeleton(TArray<TWeakObjectPtr<US
 			FString PackagePath;
 			CreateUniqueAssetName(Object->GetOutermost()->GetName(), DefaultSuffix, PackagePath, Name);
 
-			USkeletonFactory* Factory = ConstructObject<USkeletonFactory>(USkeletonFactory::StaticClass());
+			USkeletonFactory* Factory = NewObject<USkeletonFactory>();
 			Factory->TargetSkeletalMesh = Object;
 
 			FContentBrowserModule& ContentBrowserModule = FModuleManager::LoadModuleChecked<FContentBrowserModule>("ContentBrowser");
@@ -622,7 +622,7 @@ void FAssetTypeActions_SkeletalMesh::ExecuteNewSkeleton(TArray<TWeakObjectPtr<US
 				CreateUniqueAssetName(Object->GetOutermost()->GetName(), DefaultSuffix, PackageName, Name);
 
 				// Create the factory used to generate the asset
-				USkeletonFactory* Factory = ConstructObject<USkeletonFactory>(USkeletonFactory::StaticClass());
+				USkeletonFactory* Factory = NewObject<USkeletonFactory>();
 				Factory->TargetSkeletalMesh = Object;
 
 				FAssetToolsModule& AssetToolsModule = FModuleManager::GetModuleChecked<FAssetToolsModule>("AssetTools");
@@ -736,7 +736,7 @@ void FAssetTypeActions_SkeletalMesh::CreatePhysicsAssetFromMesh(USkeletalMesh* S
 
 	if( bWasOkClicked )
 	{			
-		UPhysicsAsset* NewAsset = ConstructObject<UPhysicsAsset>( UPhysicsAsset::StaticClass(), Package, *Name, RF_Public|RF_Standalone|RF_Transactional );
+		UPhysicsAsset* NewAsset = NewObject<UPhysicsAsset>(Package, *Name, RF_Public | RF_Standalone | RF_Transactional);
 		if(NewAsset)
 		{
 			// Do automatic asset generation.
diff --git a/Engine/Source/Developer/AssetTools/Private/AssetTypeActions/AssetTypeActions_SoundCue.cpp b/Engine/Source/Developer/AssetTools/Private/AssetTypeActions/AssetTypeActions_SoundCue.cpp
--- a/Engine/Source/Developer/AssetTools/Private/AssetTypeActions/AssetTypeActions_SoundCue.cpp
+++ b/Engine/Source/Developer/AssetTools/Private/AssetTypeActions/AssetTypeActions_SoundCue.cpp
@@ -83,7 +83,7 @@ void FAssetTypeActions_SoundCue::ExecuteConsolidateAttenuation(TArray<TWeakObjec
 		TArray<UObject*> ObjectsToSync;
 
 		FAssetToolsModule& AssetToolsModule = FModuleManager::GetModuleChecked<FAssetToolsModule>("AssetTools");
-		USoundAttenuationFactory* Factory = ConstructObject<USoundAttenuationFactory>(USoundAttenuationFactory::StaticClass());
+		USoundAttenuationFactory* Factory = NewObject<USoundAttenuationFactory>();
 
 		for (auto UnmatchedIt = UnmatchedAttenuations.CreateConstIterator(); UnmatchedIt; ++UnmatchedIt)
 		{
diff --git a/Engine/Source/Developer/AssetTools/Private/AssetTypeActions/AssetTypeActions_SoundWave.cpp b/Engine/Source/Developer/AssetTools/Private/AssetTypeActions/AssetTypeActions_SoundWave.cpp
--- a/Engine/Source/Developer/AssetTools/Private/AssetTypeActions/AssetTypeActions_SoundWave.cpp
+++ b/Engine/Source/Developer/AssetTools/Private/AssetTypeActions/AssetTypeActions_SoundWave.cpp
@@ -112,7 +112,7 @@ void FAssetTypeActions_SoundWave::ExecuteCreateSoundCue(TArray<TWeakObjectPtr<US
 			CreateUniqueAssetName(Object->GetOutermost()->GetName(), DefaultSuffix, PackagePath, Name);
 
 			// Create the factory used to generate the asset
-			USoundCueFactoryNew* Factory = ConstructObject<USoundCueFactoryNew>(USoundCueFactoryNew::StaticClass());
+			USoundCueFactoryNew* Factory = NewObject<USoundCueFactoryNew>();
 			Factory->InitialSoundWave = Object;
 
 			FContentBrowserModule& ContentBrowserModule = FModuleManager::LoadModuleChecked<FContentBrowserModule>("ContentBrowser");
@@ -133,7 +133,7 @@ void FAssetTypeActions_SoundWave::ExecuteCreateSoundCue(TArray<TWeakObjectPtr<US
 				CreateUniqueAssetName(Object->GetOutermost()->GetName(), DefaultSuffix, PackageName, Name);
 
 				// Create the factory used to generate the asset
-				USoundCueFactoryNew* Factory = ConstructObject<USoundCueFactoryNew>(USoundCueFactoryNew::StaticClass());
+				USoundCueFactoryNew* Factory = NewObject<USoundCueFactoryNew>();
 				Factory->InitialSoundWave = Object;
 
 				FAssetToolsModule& AssetToolsModule = FModuleManager::GetModuleChecked<FAssetToolsModule>("AssetTools");
diff --git a/Engine/Source/Developer/AssetTools/Private/AssetTypeActions/AssetTypeActions_StaticMesh.cpp b/Engine/Source/Developer/AssetTools/Private/AssetTypeActions/AssetTypeActions_StaticMesh.cpp
--- a/Engine/Source/Developer/AssetTools/Private/AssetTypeActions/AssetTypeActions_StaticMesh.cpp
+++ b/Engine/Source/Developer/AssetTools/Private/AssetTypeActions/AssetTypeActions_StaticMesh.cpp
@@ -59,7 +59,7 @@ UThumbnailInfo* FAssetTypeActions_StaticMesh::GetThumbnailInfo(UObject* Asset) c
 	UThumbnailInfo* ThumbnailInfo = StaticMesh->ThumbnailInfo;
 	if ( ThumbnailInfo == NULL )
 	{
-		ThumbnailInfo = ConstructObject<USceneThumbnailInfo>(USceneThumbnailInfo::StaticClass(), StaticMesh);
+		ThumbnailInfo = NewObject<USceneThumbnailInfo>(StaticMesh);
 		StaticMesh->ThumbnailInfo = ThumbnailInfo;
 	}
 
diff --git a/Engine/Source/Developer/AssetTools/Private/AssetTypeActions/AssetTypeActions_Texture.cpp b/Engine/Source/Developer/AssetTools/Private/AssetTypeActions/AssetTypeActions_Texture.cpp
--- a/Engine/Source/Developer/AssetTools/Private/AssetTypeActions/AssetTypeActions_Texture.cpp
+++ b/Engine/Source/Developer/AssetTools/Private/AssetTypeActions/AssetTypeActions_Texture.cpp
@@ -80,7 +80,7 @@ void FAssetTypeActions_Texture::ExecuteCreateMaterial(TArray<TWeakObjectPtr<UTex
 			CreateUniqueAssetName(Object->GetOutermost()->GetName(), DefaultSuffix, PackagePath, Name);
 
 			// Create the factory used to generate the asset
-			UMaterialFactoryNew* Factory = ConstructObject<UMaterialFactoryNew>(UMaterialFactoryNew::StaticClass());
+			UMaterialFactoryNew* Factory = NewObject<UMaterialFactoryNew>();
 			Factory->InitialTexture = Object;
 
 			FContentBrowserModule& ContentBrowserModule = FModuleManager::LoadModuleChecked<FContentBrowserModule>("ContentBrowser");
@@ -101,7 +101,7 @@ void FAssetTypeActions_Texture::ExecuteCreateMaterial(TArray<TWeakObjectPtr<UTex
 				CreateUniqueAssetName(Object->GetOutermost()->GetName(), DefaultSuffix, PackageName, Name);
 
 				// Create the factory used to generate the asset
-				UMaterialFactoryNew* Factory = ConstructObject<UMaterialFactoryNew>(UMaterialFactoryNew::StaticClass());
+				UMaterialFactoryNew* Factory = NewObject<UMaterialFactoryNew>();
 				Factory->InitialTexture = Object;
 
 				FAssetToolsModule& AssetToolsModule = FModuleManager::GetModuleChecked<FAssetToolsModule>("AssetTools");
diff --git a/Engine/Source/Developer/AssetTools/Private/AssetTypeActions/AssetTypeActions_Texture2D.cpp b/Engine/Source/Developer/AssetTools/Private/AssetTypeActions/AssetTypeActions_Texture2D.cpp
--- a/Engine/Source/Developer/AssetTools/Private/AssetTypeActions/AssetTypeActions_Texture2D.cpp
+++ b/Engine/Source/Developer/AssetTools/Private/AssetTypeActions/AssetTypeActions_Texture2D.cpp
@@ -39,7 +39,7 @@ void FAssetTypeActions_Texture2D::ExecuteCreateSlateBrush(TArray<TWeakObjectPtr<
 			CreateUniqueAssetName(Object->GetOutermost()->GetName(), DefaultSuffix, PackagePath, Name);
 
 			// Create the factory used to generate the asset
-			USlateBrushAssetFactory* Factory = ConstructObject<USlateBrushAssetFactory>(USlateBrushAssetFactory::StaticClass());
+			USlateBrushAssetFactory* Factory = NewObject<USlateBrushAssetFactory>();
 			Factory->InitialTexture = CastChecked<UTexture2D>(Object);
 			FContentBrowserModule& ContentBrowserModule = FModuleManager::LoadModuleChecked<FContentBrowserModule>("ContentBrowser");
 			ContentBrowserModule.Get().CreateNewAsset(Name, FPackageName::GetLongPackagePath(PackagePath), USlateBrushAsset::StaticClass(), Factory);
@@ -60,7 +60,7 @@ void FAssetTypeActions_Texture2D::ExecuteCreateSlateBrush(TArray<TWeakObjectPtr<
 				CreateUniqueAssetName(Object->GetOutermost()->GetName(), DefaultSuffix, PackageName, Name);
 
 				// Create the factory used to generate the asset
-				USlateBrushAssetFactory* Factory = ConstructObject<USlateBrushAssetFactory>(USlateBrushAssetFactory::StaticClass());
+				USlateBrushAssetFactory* Factory = NewObject<USlateBrushAssetFactory>();
 				Factory->InitialTexture = CastChecked<UTexture2D>(Object);
 
 				FAssetToolsModule& AssetToolsModule = FModuleManager::GetModuleChecked<FAssetToolsModule>("AssetTools");
diff --git a/Engine/Source/Developer/AssetTools/Private/AssetTypeActions/AssetTypeActions_World.cpp b/Engine/Source/Developer/AssetTools/Private/AssetTypeActions/AssetTypeActions_World.cpp
--- a/Engine/Source/Developer/AssetTools/Private/AssetTypeActions/AssetTypeActions_World.cpp
+++ b/Engine/Source/Developer/AssetTools/Private/AssetTypeActions/AssetTypeActions_World.cpp
@@ -35,7 +35,7 @@ UThumbnailInfo* FAssetTypeActions_World::GetThumbnailInfo(UObject* Asset) const
 	UThumbnailInfo* ThumbnailInfo = World->ThumbnailInfo;
 	if (ThumbnailInfo == NULL)
 	{
-		ThumbnailInfo = ConstructObject<UWorldThumbnailInfo>(UWorldThumbnailInfo::StaticClass(), World);
+		ThumbnailInfo = NewObject<UWorldThumbnailInfo>(World);
 		World->ThumbnailInfo = ThumbnailInfo;
 	}
 
diff --git a/Engine/Source/Developer/GameplayDebugger/Private/GameplayDebuggingControllerComponent.cpp b/Engine/Source/Developer/GameplayDebugger/Private/GameplayDebuggingControllerComponent.cpp
--- a/Engine/Source/Developer/GameplayDebugger/Private/GameplayDebuggingControllerComponent.cpp
+++ b/Engine/Source/Developer/GameplayDebugger/Private/GameplayDebuggingControllerComponent.cpp
@@ -217,7 +217,7 @@ void UGameplayDebuggingControllerComponent::BindAIDebugViewKeys()
 #if !(UE_BUILD_SHIPPING || UE_BUILD_TEST)
 	if (!AIDebugViewInputComponent)
 	{
-		AIDebugViewInputComponent = ConstructObject<UInputComponent>(UInputComponent::StaticClass(), GetOwner(), TEXT("AIDebugViewInputComponent0"));
+		AIDebugViewInputComponent = NewObject<UInputComponent>(GetOwner(), TEXT("AIDebugViewInputComponent0"));
 		AIDebugViewInputComponent->RegisterComponent();
 
 		AIDebugViewInputComponent->BindKey(EKeys::NumPadZero, IE_Pressed, this, &UGameplayDebuggingControllerComponent::ToggleAIDebugView_SetView0);
diff --git a/Engine/Source/Developer/GameplayDebugger/Private/GameplayDebuggingReplicator.cpp b/Engine/Source/Developer/GameplayDebugger/Private/GameplayDebuggingReplicator.cpp
--- a/Engine/Source/Developer/GameplayDebugger/Private/GameplayDebuggingReplicator.cpp
+++ b/Engine/Source/Developer/GameplayDebugger/Private/GameplayDebuggingReplicator.cpp
@@ -273,7 +273,7 @@ UGameplayDebuggingComponent* AGameplayDebuggingReplicator::GetDebugComponent()
 {
 	if (!DebugComponent && DebugComponentClass.IsValid() && GetNetMode() < ENetMode::NM_Client)
 	{
-		DebugComponent = ConstructObject<UGameplayDebuggingComponent>(DebugComponentClass.Get(), this);
+		DebugComponent = NewObject<UGameplayDebuggingComponent>(this, DebugComponentClass.Get());
 		DebugComponent->SetIsReplicated(true);
 		DebugComponent->RegisterComponent();
 		DebugComponent->Activate();
@@ -377,7 +377,7 @@ void AGameplayDebuggingReplicator::CreateTool()
 			{
 				DebugComponentControllerClass = AGameplayDebuggingHUDComponent::StaticClass();
 			}
-			GDC = ConstructObject<UGameplayDebuggingControllerComponent>(DebugComponentControllerClass.Get(), this);
+			GDC = NewObject<UGameplayDebuggingControllerComponent>(this, DebugComponentControllerClass.Get());
 			GDC->SetPlayerOwner(LocalPlayerOwner);
 			GDC->RegisterComponent();
 		}
diff --git a/Engine/Source/Developer/LogVisualizer/Private/VisualLoggerRenderingActor.cpp b/Engine/Source/Developer/LogVisualizer/Private/VisualLoggerRenderingActor.cpp
--- a/Engine/Source/Developer/LogVisualizer/Private/VisualLoggerRenderingActor.cpp
+++ b/Engine/Source/Developer/LogVisualizer/Private/VisualLoggerRenderingActor.cpp
@@ -105,7 +105,7 @@ void UVisualLoggerRenderingComponent::DestroyRenderState_Concurrent()
 AVisualLoggerRenderingActor::AVisualLoggerRenderingActor(const FObjectInitializer& ObjectInitializer)
 	: Super(ObjectInitializer)
 {
-	RenderingComponent = ConstructObject<UVisualLoggerRenderingComponent>(UVisualLoggerRenderingComponent::StaticClass(), this);
+	RenderingComponent = CreateDefaultSubobject<UVisualLoggerRenderingComponent>(TEXT("RenderingComponent"));
 }
 
 void AVisualLoggerRenderingActor::AddDebugRendering()
diff --git a/Engine/Source/Developer/MovieSceneTools/Private/K2Node_PlayMovieScene.cpp b/Engine/Source/Developer/MovieSceneTools/Private/K2Node_PlayMovieScene.cpp
--- a/Engine/Source/Developer/MovieSceneTools/Private/K2Node_PlayMovieScene.cpp
+++ b/Engine/Source/Developer/MovieSceneTools/Private/K2Node_PlayMovieScene.cpp
@@ -191,7 +191,7 @@ void UK2Node_PlayMovieScene::CreateBindingsIfNeeded()
 
 		// The new MovieSceneBindings object will be stored as an "inner" of this object
 		UObject* MovieSceneBindingsOuter = this;
-		UMovieSceneBindings* NewMovieSceneBindings = ConstructObject< UMovieSceneBindings >( UMovieSceneBindings::StaticClass(), MovieSceneBindingsOuter );
+		UMovieSceneBindings* NewMovieSceneBindings = NewObject<UMovieSceneBindings>(MovieSceneBindingsOuter);
 		check( NewMovieSceneBindings != NULL );
 
 		MovieSceneBindings = NewMovieSceneBindings;
diff --git a/Engine/Source/Developer/Windows/WindowsTargetPlatform/Private/WindowsTargetPlatformModule.cpp b/Engine/Source/Developer/Windows/WindowsTargetPlatform/Private/WindowsTargetPlatformModule.cpp
--- a/Engine/Source/Developer/Windows/WindowsTargetPlatform/Private/WindowsTargetPlatformModule.cpp
+++ b/Engine/Source/Developer/Windows/WindowsTargetPlatform/Private/WindowsTargetPlatformModule.cpp
@@ -65,7 +65,7 @@ public:
 		FCoreDelegates::GetHotfixDelegate(EHotfixDelegates::Test).BindRaw(this, &FWindowsTargetPlatformModule::HotfixTest);
 #endif
 
-		TargetSettings = ConstructObject<UWindowsTargetSettings>(UWindowsTargetSettings::StaticClass(), GetTransientPackage(), "WindowsTargetSettings", RF_Standalone);
+		TargetSettings = NewObject<UWindowsTargetSettings>(GetTransientPackage(), "WindowsTargetSettings", RF_Standalone);
 		
 		// We need to manually load the config properties here, as this module is loaded before the UObject system is setup to do this
 		GConfig->GetArray(TEXT("/Script/WindowsTargetPlatform.WindowsTargetSettings"), TEXT("TargetedRHIs"), TargetSettings->TargetedRHIs, GEngineIni);
diff --git a/Engine/Source/Editor/BehaviorTreeEditor/Private/BehaviorTreeDecoratorGraphNode_Decorator.cpp b/Engine/Source/Editor/BehaviorTreeEditor/Private/BehaviorTreeDecoratorGraphNode_Decorator.cpp
--- a/Engine/Source/Editor/BehaviorTreeEditor/Private/BehaviorTreeDecoratorGraphNode_Decorator.cpp
+++ b/Engine/Source/Editor/BehaviorTreeEditor/Private/BehaviorTreeDecoratorGraphNode_Decorator.cpp
@@ -27,7 +27,7 @@ void UBehaviorTreeDecoratorGraphNode_Decorator::PostPlacedNewNode()
 			UBehaviorTree* BT = Cast<UBehaviorTree>(OwningNode->GetOuter()->GetOuter());
 			if (BT)
 			{
-				UBTDecorator* MyDecorator = ConstructObject<UBTDecorator>(NodeClass, BT);
+				UBTDecorator* MyDecorator = NewObject<UBTDecorator>(BT, NodeClass);
 				MyDecorator->InitializeFromAsset(*BT);
 				OwningNode->InitializeDecorator(MyDecorator);
 
diff --git a/Engine/Source/Editor/BehaviorTreeEditor/Private/BehaviorTreeEditor.cpp b/Engine/Source/Editor/BehaviorTreeEditor/Private/BehaviorTreeEditor.cpp
--- a/Engine/Source/Editor/BehaviorTreeEditor/Private/BehaviorTreeEditor.cpp
+++ b/Engine/Source/Editor/BehaviorTreeEditor/Private/BehaviorTreeEditor.cpp
@@ -2185,7 +2185,7 @@ void FBehaviorTreeEditor::CreateNewBlackboard()
 	FAssetToolsModule& AssetToolsModule = FModuleManager::GetModuleChecked<FAssetToolsModule>("AssetTools");
 	AssetToolsModule.Get().CreateUniqueAssetName(PathNameWithFilename, TEXT(""), PackageName, Name);
 
-	UDataAssetFactory* DataAssetFactory = ConstructObject<UDataAssetFactory>(UDataAssetFactory::StaticClass());
+	UDataAssetFactory* DataAssetFactory = NewObject<UDataAssetFactory>();
 	FContentBrowserModule& ContentBrowserModule = FModuleManager::LoadModuleChecked<FContentBrowserModule>("ContentBrowser");
 	ContentBrowserModule.Get().CreateNewAsset(Name, PathName, UBlackboardData::StaticClass(), DataAssetFactory);
 }
diff --git a/Engine/Source/Editor/BehaviorTreeEditor/Private/BehaviorTreeFactory.cpp b/Engine/Source/Editor/BehaviorTreeEditor/Private/BehaviorTreeFactory.cpp
--- a/Engine/Source/Editor/BehaviorTreeEditor/Private/BehaviorTreeFactory.cpp
+++ b/Engine/Source/Editor/BehaviorTreeEditor/Private/BehaviorTreeFactory.cpp
@@ -22,7 +22,7 @@ bool UBehaviorTreeFactory::CanCreateNew() const
 UObject* UBehaviorTreeFactory::FactoryCreateNew(UClass* Class,UObject* InParent,FName Name,EObjectFlags Flags,UObject* Context,FFeedbackContext* Warn)
 {
 	check(Class->IsChildOf(UBehaviorTree::StaticClass()));
-	return ConstructObject<UBehaviorTree>(Class, InParent, Name, Flags);;
+	return NewObject<UBehaviorTree>(InParent, Class, Name, Flags);;
 }
 
 #undef LOCTEXT_NAMESPACE
\ No newline at end of file
diff --git a/Engine/Source/Editor/BehaviorTreeEditor/Private/BehaviorTreeGraphNode.cpp b/Engine/Source/Editor/BehaviorTreeEditor/Private/BehaviorTreeGraphNode.cpp
--- a/Engine/Source/Editor/BehaviorTreeEditor/Private/BehaviorTreeGraphNode.cpp
+++ b/Engine/Source/Editor/BehaviorTreeEditor/Private/BehaviorTreeGraphNode.cpp
@@ -52,7 +52,7 @@ void UBehaviorTreeGraphNode::PostPlacedNewNode()
 		UBehaviorTree* BT = Cast<UBehaviorTree>(GetBehaviorTreeGraph()->GetOuter());
 		if (BT)
 		{
-			NodeInstance = ConstructObject<UBTNode>(NodeClass, BT);
+			NodeInstance = NewObject<UBTNode>(BT, NodeClass);
 
 			UBTNode* BTNode = (UBTNode*)NodeInstance;
 			BTNode->SetFlags(RF_Transactional);
diff --git a/Engine/Source/Editor/BehaviorTreeEditor/Private/BlackboardDataFactory.cpp b/Engine/Source/Editor/BehaviorTreeEditor/Private/BlackboardDataFactory.cpp
--- a/Engine/Source/Editor/BehaviorTreeEditor/Private/BlackboardDataFactory.cpp
+++ b/Engine/Source/Editor/BehaviorTreeEditor/Private/BlackboardDataFactory.cpp
@@ -22,7 +22,7 @@ bool UBlackboardDataFactory::CanCreateNew() const
 UObject* UBlackboardDataFactory::FactoryCreateNew(UClass* Class,UObject* InParent,FName Name,EObjectFlags Flags,UObject* Context,FFeedbackContext* Warn)
 {
 	check(Class->IsChildOf(UBlackboardData::StaticClass()));
-	return ConstructObject<UBlackboardData>(Class, InParent, Name, Flags);
+	return NewObject<UBlackboardData>(InParent, Class, Name, Flags);
 }
 
 #undef LOCTEXT_NAMESPACE
\ No newline at end of file
diff --git a/Engine/Source/Editor/BehaviorTreeEditor/Private/EdGraphSchema_BehaviorTree.cpp b/Engine/Source/Editor/BehaviorTreeEditor/Private/EdGraphSchema_BehaviorTree.cpp
--- a/Engine/Source/Editor/BehaviorTreeEditor/Private/EdGraphSchema_BehaviorTree.cpp
+++ b/Engine/Source/Editor/BehaviorTreeEditor/Private/EdGraphSchema_BehaviorTree.cpp
@@ -286,7 +286,7 @@ void UEdGraphSchema_BehaviorTree::GetGraphContextActions(FGraphContextMenuBuilde
 				GraphNodeClass = UBehaviorTreeGraphNode_SubtreeTask::StaticClass();
 			}
 
-			UBehaviorTreeGraphNode* OpNode = ConstructObject<UBehaviorTreeGraphNode>(GraphNodeClass, ContextMenuBuilder.OwnerOfTemporaries);
+			UBehaviorTreeGraphNode* OpNode = NewObject<UBehaviorTreeGraphNode>(ContextMenuBuilder.OwnerOfTemporaries, GraphNodeClass);
 			OpNode->ClassData = NodeClass;
 			AddOpAction->NodeTemplate = OpNode;
 		}
diff --git a/Engine/Source/Editor/BehaviorTreeEditor/Private/SBehaviorTreeBlackboardEditor.cpp b/Engine/Source/Editor/BehaviorTreeEditor/Private/SBehaviorTreeBlackboardEditor.cpp
--- a/Engine/Source/Editor/BehaviorTreeEditor/Private/SBehaviorTreeBlackboardEditor.cpp
+++ b/Engine/Source/Editor/BehaviorTreeEditor/Private/SBehaviorTreeBlackboardEditor.cpp
@@ -249,7 +249,7 @@ void SBehaviorTreeBlackboardEditor::HandleKeyClassPicked(UClass* InClass)
 
 	FBlackboardEntry Entry;
 	Entry.EntryName = FName(*NewKeyName);
-	Entry.KeyType = ConstructObject<UBlackboardKeyType>(InClass, BlackboardData);		
+	Entry.KeyType = NewObject<UBlackboardKeyType>(BlackboardData, InClass);
 
 	BlackboardData->Keys.Add(Entry);
 
diff --git a/Engine/Source/Editor/BlueprintGraph/Private/BlueprintComponentNodeSpawner.cpp b/Engine/Source/Editor/BlueprintGraph/Private/BlueprintComponentNodeSpawner.cpp
--- a/Engine/Source/Editor/BlueprintGraph/Private/BlueprintComponentNodeSpawner.cpp
+++ b/Engine/Source/Editor/BlueprintGraph/Private/BlueprintComponentNodeSpawner.cpp
@@ -118,7 +118,7 @@ UEdGraphNode* UBlueprintComponentNodeSpawner::Invoke(UEdGraph* ParentGraph, FBin
 	if (!bIsTemplateNode)
 	{
 		UBlueprint* Blueprint = NewNode->GetBlueprint();
-		UActorComponent* ComponentTemplate = ConstructObject<UActorComponent>(ComponentClass, Blueprint->GeneratedClass);
+		UActorComponent* ComponentTemplate = NewObject<UActorComponent>(Blueprint->GeneratedClass, ComponentClass);
 		ComponentTemplate->SetFlags(RF_ArchetypeObject);
 
 		Blueprint->ComponentTemplates.Add(ComponentTemplate);
diff --git a/Engine/Source/Editor/BlueprintGraph/Private/BlueprintNodeTemplateCache.cpp b/Engine/Source/Editor/BlueprintGraph/Private/BlueprintNodeTemplateCache.cpp
--- a/Engine/Source/Editor/BlueprintGraph/Private/BlueprintNodeTemplateCache.cpp
+++ b/Engine/Source/Editor/BlueprintGraph/Private/BlueprintNodeTemplateCache.cpp
@@ -176,7 +176,7 @@ static UEdGraph* BlueprintNodeTemplateCacheImpl::AddGraph(UBlueprint* BlueprintO
 	UClass* GraphClass = UEdGraph::StaticClass();
 	FName const GraphName = MakeUniqueObjectName(BlueprintOuter, GraphClass, FName(TEXT("TEMPLATE_NODE_OUTER")));
 
-	UEdGraph* NewGraph = ConstructObject<UEdGraph>(GraphClass, BlueprintOuter, GraphName, RF_Transient);
+	UEdGraph* NewGraph = NewObject<UEdGraph>(BlueprintOuter, GraphClass, GraphName, RF_Transient);
 	NewGraph->Schema = SchemaClass;
 
 	MarkGraphForTemplateUse(NewGraph);
diff --git a/Engine/Source/Editor/BlueprintGraph/Private/EdGraphSchema_K2.cpp b/Engine/Source/Editor/BlueprintGraph/Private/EdGraphSchema_K2.cpp
--- a/Engine/Source/Editor/BlueprintGraph/Private/EdGraphSchema_K2.cpp
+++ b/Engine/Source/Editor/BlueprintGraph/Private/EdGraphSchema_K2.cpp
@@ -4923,7 +4923,7 @@ UEdGraphNode* UEdGraphSchema_K2::CreateSubstituteNode(UEdGraphNode* Node, const
 			}
 
 			// Create a custom event node to replace the original event node imported from text
-			UK2Node_CustomEvent* CustomEventNode = ConstructObject<UK2Node_CustomEvent>(UK2Node_CustomEvent::StaticClass(), EventNode->GetOuter(), ObjName, EventNode->GetFlags(), NULL, true, InstanceGraph);
+			UK2Node_CustomEvent* CustomEventNode = NewObject<UK2Node_CustomEvent>(EventNode->GetOuter(), ObjName, EventNode->GetFlags(), nullptr, true, InstanceGraph);
 
 			// Ensure that it is editable
 			CustomEventNode->bIsEditable = true;
diff --git a/Engine/Source/Editor/BlueprintGraph/Private/EdGraphSchema_K2_Actions.cpp b/Engine/Source/Editor/BlueprintGraph/Private/EdGraphSchema_K2_Actions.cpp
--- a/Engine/Source/Editor/BlueprintGraph/Private/EdGraphSchema_K2_Actions.cpp
+++ b/Engine/Source/Editor/BlueprintGraph/Private/EdGraphSchema_K2_Actions.cpp
@@ -303,7 +303,7 @@ UEdGraphNode* FEdGraphSchemaAction_K2AddComponent::PerformAction(class UEdGraph*
 
 		ensure(NULL != Cast<UBlueprintGeneratedClass>(Blueprint->GeneratedClass));
 		// Then create a new template object, and add to array in
-		UActorComponent* NewTemplate = ConstructObject<UActorComponent>(ComponentClass, Blueprint->GeneratedClass);
+		UActorComponent* NewTemplate = NewObject<UActorComponent>(Blueprint->GeneratedClass, ComponentClass);
 		NewTemplate->SetFlags(RF_ArchetypeObject);
 		Blueprint->ComponentTemplates.Add(NewTemplate);
 
diff --git a/Engine/Source/Editor/BlueprintGraph/Private/K2Node_AddComponent.cpp b/Engine/Source/Editor/BlueprintGraph/Private/K2Node_AddComponent.cpp
--- a/Engine/Source/Editor/BlueprintGraph/Private/K2Node_AddComponent.cpp
+++ b/Engine/Source/Editor/BlueprintGraph/Private/K2Node_AddComponent.cpp
@@ -239,7 +239,7 @@ void UK2Node_AddComponent::PostPasteNode()
 		{
 			ensure(NULL != Cast<UBlueprintGeneratedClass>(Blueprint->GeneratedClass));
 			// Create a new template object and update the template pin to point to it
-			UActorComponent* NewTemplate = ConstructObject<UActorComponent>(ComponentClass, Blueprint->GeneratedClass);
+			UActorComponent* NewTemplate = NewObject<UActorComponent>(Blueprint->GeneratedClass, ComponentClass);
 			NewTemplate->SetFlags(RF_ArchetypeObject);
 			Blueprint->ComponentTemplates.Add(NewTemplate);
 
diff --git a/Engine/Source/Editor/Cascade/Private/Cascade.cpp b/Engine/Source/Editor/Cascade/Private/Cascade.cpp
--- a/Engine/Source/Editor/Cascade/Private/Cascade.cpp
+++ b/Engine/Source/Editor/Cascade/Private/Cascade.cpp
@@ -185,9 +185,9 @@ void FCascade::InitCascade(const EToolkitMode::Type Mode, const TSharedPtr< clas
 
 	CurrentLODIdx = 0;
 
-	EditorOptions = ConstructObject<UCascadeOptions>(UCascadeOptions::StaticClass());
+	EditorOptions = NewObject<UCascadeOptions>();
 	check(EditorOptions);
-	EditorConfig = ConstructObject<UCascadeConfiguration>(UCascadeConfiguration::StaticClass());
+	EditorConfig = NewObject<UCascadeConfiguration>();
 	check(EditorConfig);
 
 	FString Description;
@@ -237,9 +237,9 @@ void FCascade::InitCascade(const EToolkitMode::Type Mode, const TSharedPtr< clas
 		FEngineAnalytics::GetProvider().RecordEvent(TEXT("Editor.Usage.Cascade.Init"), TEXT("Overview"), Description);
 	}
 
-	ParticleSystemComponent = ConstructObject<UCascadeParticleSystemComponent>(UCascadeParticleSystemComponent::StaticClass());
+	ParticleSystemComponent = NewObject<UCascadeParticleSystemComponent>();
 
-	LocalVectorFieldPreviewComponent = ConstructObject<UVectorFieldComponent>(UVectorFieldComponent::StaticClass());
+	LocalVectorFieldPreviewComponent = NewObject<UVectorFieldComponent>();
 
 	bIsSoloing = false;
 
@@ -924,7 +924,7 @@ void FCascade::OnNewModule(int32 Idx)
 	ParticleSystemComponent->PreEditChange(NULL);
 
 	// Construct it and add to selected emitter.
-	UParticleModule* NewModule = ConstructObject<UParticleModule>(NewModClass, ParticleSystem, NAME_None, RF_Transactional);
+	UParticleModule* NewModule = NewObject<UParticleModule>(ParticleSystem, NewModClass, NAME_None, RF_Transactional);
 	NewModule->ModuleEditorColor = FColor::MakeRandomColor();
 	NewModule->SetToSensibleDefaults(SelectedEmitter);
 	NewModule->LODValidity = 1;
@@ -996,7 +996,7 @@ void FCascade::OnNewEmitter()
 	UClass* NewEmitClass = UParticleSpriteEmitter::StaticClass();
 
 	// Construct it
-	UParticleEmitter* NewEmitter = ConstructObject<UParticleEmitter>(NewEmitClass, ParticleSystem, NAME_None, RF_Transactional);
+	UParticleEmitter* NewEmitter = NewObject<UParticleEmitter>(ParticleSystem, NewEmitClass, NAME_None, RF_Transactional);
 	UParticleLODLevel* LODLevel	= NewEmitter->GetLODLevel(0);
 	if (LODLevel == NULL)
 	{
@@ -1974,7 +1974,7 @@ void FCascade::CreateInternalWidgets()
 
 	if (!ParticleSystem->CurveEdSetup)
 	{
-		ParticleSystem->CurveEdSetup = ConstructObject<UInterpCurveEdSetup>(UInterpCurveEdSetup::StaticClass(), ParticleSystem, NAME_None, RF_Transactional);
+		ParticleSystem->CurveEdSetup = NewObject<UInterpCurveEdSetup>(ParticleSystem, NAME_None, RF_Transactional);
 	}
 
 	IDistributionCurveEditorModule* CurveEditorModule = &FModuleManager::LoadModuleChecked<IDistributionCurveEditorModule>( "DistCurveEditor" );
@@ -2676,7 +2676,7 @@ bool FCascade::DuplicateEmitter(UParticleEmitter* SourceEmitter, UParticleSystem
 	if (NewEmitClass == UParticleSpriteEmitter::StaticClass())
 	{
 		// Construct it
-		UParticleEmitter* NewEmitter = ConstructObject<UParticleEmitter>(NewEmitClass, DestSystem, NAME_None, RF_Transactional);
+		UParticleEmitter* NewEmitter = NewObject<UParticleEmitter>(DestSystem, NewEmitClass, NAME_None, RF_Transactional);
 
 		check(NewEmitter);
 
@@ -2696,7 +2696,7 @@ bool FCascade::DuplicateEmitter(UParticleEmitter* SourceEmitter, UParticleSystem
 		for (int32 LODIndex = 0; LODIndex < SourceEmitter->LODLevels.Num(); LODIndex++)
 		{
 			SourceLODLevel	= SourceEmitter->LODLevels[LODIndex];
-			NewLODLevel		= ConstructObject<UParticleLODLevel>(UParticleLODLevel::StaticClass(), NewEmitter, NAME_None, RF_Transactional);
+			NewLODLevel = NewObject<UParticleLODLevel>(NewEmitter, NAME_None, RF_Transactional);
 			check(NewLODLevel);
 
 			NewLODLevel->Level					= SourceLODLevel->Level;
diff --git a/Engine/Source/Editor/Cascade/Private/CascadePreviewViewportClient.cpp b/Engine/Source/Editor/Cascade/Private/CascadePreviewViewportClient.cpp
--- a/Engine/Source/Editor/Cascade/Private/CascadePreviewViewportClient.cpp
+++ b/Engine/Source/Editor/Cascade/Private/CascadePreviewViewportClient.cpp
@@ -147,7 +147,7 @@ FCascadeEdPreviewViewportClient::FCascadeEdPreviewViewportClient(TWeakPtr<FCasca
 
 	if (Mesh)
 	{
-		FloorComponent = ConstructObject<UStaticMeshComponent>(UStaticMeshComponent::StaticClass());
+		FloorComponent = NewObject<UStaticMeshComponent>(GetTransientPackage(), TEXT("FloorComponent"));
 		check(FloorComponent);
 		FloorComponent->StaticMesh = Mesh;
 		FloorComponent->DepthPriorityGroup = SDPG_World;
diff --git a/Engine/Source/Editor/ContentBrowser/Private/SContentBrowser.cpp b/Engine/Source/Editor/ContentBrowser/Private/SContentBrowser.cpp
--- a/Engine/Source/Editor/ContentBrowser/Private/SContentBrowser.cpp
+++ b/Engine/Source/Editor/ContentBrowser/Private/SContentBrowser.cpp
@@ -1099,7 +1099,7 @@ void SContentBrowser::NewAssetRequested(const FString& SelectedPath, TWeakObject
 {
 	if ( ensure(SelectedPath.Len() > 0) && ensure(FactoryClass.IsValid()) )
 	{
-		UFactory* NewFactory = ConstructObject<UFactory>( FactoryClass.Get() );
+		UFactory* NewFactory = NewObject<UFactory>(GetTransientPackage(), FactoryClass.Get());
 		FEditorDelegates::OnConfigureNewAssetProperties.Broadcast(NewFactory);
 		if ( NewFactory->ConfigureProperties() )
 		{
diff --git a/Engine/Source/Editor/DestructibleMeshEditor/Private/SDestructibleMeshEditorViewport.cpp b/Engine/Source/Editor/DestructibleMeshEditor/Private/SDestructibleMeshEditorViewport.cpp
--- a/Engine/Source/Editor/DestructibleMeshEditor/Private/SDestructibleMeshEditorViewport.cpp
+++ b/Engine/Source/Editor/DestructibleMeshEditor/Private/SDestructibleMeshEditorViewport.cpp
@@ -162,7 +162,7 @@ void FDestructibleMeshEditorViewportClient::UpdateChunkSelection( TArray<int32>
 	// make sure we have enough proxies to fill the selection array */
 	while(UnusedProxies.Num() < InSelectedChunkIndices.Num())
 	{
-		UnusedProxies.Add(ConstructObject<UDestructibleChunkParamsProxy>(UDestructibleChunkParamsProxy::StaticClass(), GetTransientPackage(), NAME_None, RF_NoFlags, NULL, false, NULL));
+		UnusedProxies.Add(NewObject<UDestructibleChunkParamsProxy>());
 	}
 
 	UDestructibleMesh* DestructibleMesh = DestructibleMeshEditorPtr.Pin()->GetDestructibleMesh();
@@ -439,7 +439,7 @@ void FDestructibleMeshEditorViewportClient::ImportFBXChunks()
 				TArray<FbxNode*> FbxMeshArray;
 				FFbxImporter->FillFbxMeshArray(FFbxImporter->Scene->GetRootNode(), FbxMeshArray, FFbxImporter);
 
- 				UFbxStaticMeshImportData* ImportData = ConstructObject<UFbxStaticMeshImportData>(UFbxStaticMeshImportData::StaticClass(), GetTransientPackage(), NAME_None, RF_NoFlags, NULL);
+				UFbxStaticMeshImportData* ImportData = NewObject<UFbxStaticMeshImportData>(GetTransientPackage(), NAME_None, RF_NoFlags, NULL);
 
 				TArray<UStaticMesh*> ChunkMeshes;
 
@@ -483,8 +483,7 @@ void SDestructibleMeshEditorViewport::Construct(const FArguments& InArgs)
 
 	SEditorViewport::Construct(SEditorViewport::FArguments());
 
-	PreviewComponent = ConstructObject<UDestructibleComponent>(
-		UDestructibleComponent::StaticClass(), GetTransientPackage(), NAME_None, RF_Transient );
+	PreviewComponent = NewObject<UDestructibleComponent>(GetTransientPackage(), NAME_None, RF_Transient );
 
 	SetPreviewMesh(InArgs._ObjectToEdit);
 
@@ -606,7 +605,7 @@ void SDestructibleMeshEditorViewport::UpdatePreviewMesh(UDestructibleMesh* InDes
 
 	DestructibleMesh = InDestructibleMesh;
 
-	PreviewComponent = ConstructObject<UDestructibleComponent>(UDestructibleComponent::StaticClass());
+	PreviewComponent = NewObject<UDestructibleComponent>();
 
 	PreviewComponent->SetSkeletalMesh(InDestructibleMesh);
 
diff --git a/Engine/Source/Editor/DetailCustomizations/Private/AmbientSoundDetails.cpp b/Engine/Source/Editor/DetailCustomizations/Private/AmbientSoundDetails.cpp
--- a/Engine/Source/Editor/DetailCustomizations/Private/AmbientSoundDetails.cpp
+++ b/Engine/Source/Editor/DetailCustomizations/Private/AmbientSoundDetails.cpp
@@ -167,7 +167,7 @@ void FAmbientSoundDetails::CreateNewSoundCue( ESoundCueLayouts Layout )
 		int32 NodeColumn = 0;
 		USoundNode* PrevNode = NULL;
 
-		SoundCue = ConstructObject<USoundCue>(USoundCue::StaticClass(), AS, FName(*AS->GetInternalSoundCueName()));
+		SoundCue = NewObject<USoundCue>(AS, FName(*AS->GetInternalSoundCueName()));
 		AS->GetAudioComponent()->Sound = SoundCue;
 		AS->GetAudioComponent()->PostEditChange();
 
diff --git a/Engine/Source/Editor/DetailCustomizations/Private/AnimMontageSegmentDetails.cpp b/Engine/Source/Editor/DetailCustomizations/Private/AnimMontageSegmentDetails.cpp
--- a/Engine/Source/Editor/DetailCustomizations/Private/AnimMontageSegmentDetails.cpp
+++ b/Engine/Source/Editor/DetailCustomizations/Private/AnimMontageSegmentDetails.cpp
@@ -202,7 +202,7 @@ void SAnimationSegmentViewport::Construct(const FArguments& InArgs)
 
 	ViewportWidget->SetViewportInterface( SceneViewport.ToSharedRef() );
 	
-	PreviewComponent = ConstructObject<UDebugSkelMeshComponent>(UDebugSkelMeshComponent::StaticClass());
+	PreviewComponent = NewObject<UDebugSkelMeshComponent>();
 	PreviewComponent->bEnablePhysicsOnDedicatedServer = EMeshComponentUpdateFlag::AlwaysTickPoseAndRefreshBones;
 	PreviewScene.AddComponent(PreviewComponent, FTransform::Identity);
 
diff --git a/Engine/Source/Editor/DetailCustomizations/Private/AnimSequenceDetails.cpp b/Engine/Source/Editor/DetailCustomizations/Private/AnimSequenceDetails.cpp
--- a/Engine/Source/Editor/DetailCustomizations/Private/AnimSequenceDetails.cpp
+++ b/Engine/Source/Editor/DetailCustomizations/Private/AnimSequenceDetails.cpp
@@ -341,7 +341,7 @@ void SAnimationRefPoseViewport::Construct(const FArguments& InArgs)
 	RefFrameIndexPropertyHandle = InArgs._RefFrameIndexPropertyHandle;
 
 	// Create the preview component
-	PreviewComponent = ConstructObject<UDebugSkelMeshComponent>( UDebugSkelMeshComponent::StaticClass() );
+	PreviewComponent = NewObject<UDebugSkelMeshComponent>();
 	PreviewComponent->MeshComponentUpdateFlag = EMeshComponentUpdateFlag::AlwaysTickPoseAndRefreshBones;
 	PreviewScene.AddComponent( PreviewComponent, FTransform::Identity );
 
diff --git a/Engine/Source/Editor/DetailCustomizations/Private/BrushDetails.cpp b/Engine/Source/Editor/DetailCustomizations/Private/BrushDetails.cpp
--- a/Engine/Source/Editor/DetailCustomizations/Private/BrushDetails.cpp
+++ b/Engine/Source/Editor/DetailCustomizations/Private/BrushDetails.cpp
@@ -128,14 +128,14 @@ void FBrushDetails::CustomizeDetails( IDetailLayoutBuilder& DetailLayout )
 				const FScopedTransaction Transaction(NSLOCTEXT("UnrealEd", "BrushSet", "Brush Set"));
 				for (UObject* OuterObject : OuterObjects)
 				{
-					UBrushBuilder* NewObject = ConstructObject<UBrushBuilder>(InChosenClass, OuterObject, NAME_None, RF_Transactional);
+					UBrushBuilder* NewUObject = NewObject<UBrushBuilder>(OuterObject, InChosenClass, NAME_None, RF_Transactional);
 
 					FNewBrushBuilder NewBuilder;
-					NewBuilder.Builder = NewObject;
+					NewBuilder.Builder = NewUObject;
 					NewBuilder.Brush = CastChecked<ABrush>(OuterObject);
 
 					NewBuilders.Add(NewBuilder);
-					NewObjectPaths.Add(NewObject->GetPathName());
+					NewObjectPaths.Add(NewUObject->GetPathName());
 				}
 
 				BrushBuilderHandle->SetPerObjectValues(NewObjectPaths);
diff --git a/Engine/Source/Editor/DistCurveEditor/Private/CurveEditorSharedData.cpp b/Engine/Source/Editor/DistCurveEditor/Private/CurveEditorSharedData.cpp
--- a/Engine/Source/Editor/DistCurveEditor/Private/CurveEditorSharedData.cpp
+++ b/Engine/Source/Editor/DistCurveEditor/Private/CurveEditorSharedData.cpp
@@ -9,7 +9,7 @@
 FCurveEditorSharedData::FCurveEditorSharedData(UInterpCurveEdSetup* InEdSetup)
 	: LabelEntryHeight(36)
 {
-	EditorOptions = ConstructObject<UCurveEdOptions>(UCurveEdOptions::StaticClass());
+	EditorOptions = NewObject<UCurveEdOptions>(GetTransientPackage(), TEXT("EditorOptions"));
 	check(EditorOptions);
 	MinViewRange = EditorOptions->MinViewRange;
 	MaxViewRange = EditorOptions->MaxViewRange;
diff --git a/Engine/Source/Editor/EnvironmentQueryEditor/Private/EnvironmentQueryEditor.cpp b/Engine/Source/Editor/EnvironmentQueryEditor/Private/EnvironmentQueryEditor.cpp
--- a/Engine/Source/Editor/EnvironmentQueryEditor/Private/EnvironmentQueryEditor.cpp
+++ b/Engine/Source/Editor/EnvironmentQueryEditor/Private/EnvironmentQueryEditor.cpp
@@ -218,7 +218,7 @@ TSharedRef<SDockTab> FEnvironmentQueryEditor::SpawnTab_UpdateGraph( const FSpawn
 	check( Args.GetTabId().TabType == EQSUpdateGraphTabId );
 	if (Query->EdGraph == NULL)
 	{
-		UEnvironmentQueryGraph* MyGraph = ConstructObject<UEnvironmentQueryGraph>(UEnvironmentQueryGraph::StaticClass(), Query, NAME_None, RF_Transactional);
+		UEnvironmentQueryGraph* MyGraph = NewObject<UEnvironmentQueryGraph>(Query, NAME_None, RF_Transactional);
 		MyGraph->MarkVersion();
 
 		Query->EdGraph = MyGraph;
diff --git a/Engine/Source/Editor/EnvironmentQueryEditor/Private/EnvironmentQueryFactory.cpp b/Engine/Source/Editor/EnvironmentQueryEditor/Private/EnvironmentQueryFactory.cpp
--- a/Engine/Source/Editor/EnvironmentQueryEditor/Private/EnvironmentQueryFactory.cpp
+++ b/Engine/Source/Editor/EnvironmentQueryEditor/Private/EnvironmentQueryFactory.cpp
@@ -15,7 +15,7 @@ UEnvironmentQueryFactory::UEnvironmentQueryFactory(const FObjectInitializer& Obj
 UObject* UEnvironmentQueryFactory::FactoryCreateNew(UClass* Class,UObject* InParent,FName Name,EObjectFlags Flags,UObject* Context,FFeedbackContext* Warn)
 {
 	check(Class->IsChildOf(UEnvQuery::StaticClass()));
-	return ConstructObject<UEnvQuery>(Class, InParent, Name, Flags);
+	return NewObject<UEnvQuery>(InParent, Class, Name, Flags);
 }
 
 bool UEnvironmentQueryFactory::CanCreateNew() const
diff --git a/Engine/Source/Editor/EnvironmentQueryEditor/Private/EnvironmentQueryGraphNode_Option.cpp b/Engine/Source/Editor/EnvironmentQueryEditor/Private/EnvironmentQueryGraphNode_Option.cpp
--- a/Engine/Source/Editor/EnvironmentQueryEditor/Private/EnvironmentQueryGraphNode_Option.cpp
+++ b/Engine/Source/Editor/EnvironmentQueryEditor/Private/EnvironmentQueryGraphNode_Option.cpp
@@ -23,8 +23,8 @@ void UEnvironmentQueryGraphNode_Option::PostPlacedNewNode()
 	if (EnvQueryNodeClass != NULL)
 	{
 		UEnvQuery* Query = Cast<UEnvQuery>(GetEnvironmentQueryGraph()->GetOuter());
-		UEnvQueryOption* QueryOption = ConstructObject<UEnvQueryOption>(UEnvQueryOption::StaticClass(), Query);
-		QueryOption->Generator = ConstructObject<UEnvQueryGenerator>(EnvQueryNodeClass, Query);
+		UEnvQueryOption* QueryOption = NewObject<UEnvQueryOption>(Query);
+		QueryOption->Generator = NewObject<UEnvQueryGenerator>(Query, EnvQueryNodeClass);
 		QueryOption->Generator->UpdateGeneratorVersion();
 		
 		QueryOption->SetFlags(RF_Transactional);
diff --git a/Engine/Source/Editor/EnvironmentQueryEditor/Private/EnvironmentQueryGraphNode_Test.cpp b/Engine/Source/Editor/EnvironmentQueryEditor/Private/EnvironmentQueryGraphNode_Test.cpp
--- a/Engine/Source/Editor/EnvironmentQueryEditor/Private/EnvironmentQueryGraphNode_Test.cpp
+++ b/Engine/Source/Editor/EnvironmentQueryEditor/Private/EnvironmentQueryGraphNode_Test.cpp
@@ -15,7 +15,7 @@ void UEnvironmentQueryGraphNode_Test::PostPlacedNewNode()
 	if (EnvQueryNodeClass != NULL)
 	{
 		UEnvQuery* Query = Cast<UEnvQuery>(GetEnvironmentQueryGraph()->GetOuter());
-		NodeInstance = ConstructObject<UEnvQueryTest>(EnvQueryNodeClass, Query);
+		NodeInstance = NewObject<UEnvQueryTest>(Query, EnvQueryNodeClass);
 		NodeInstance->SetFlags(RF_Transactional);
 
 		UEnvQueryTest* TestInstance = Cast<UEnvQueryTest>(NodeInstance);
diff --git a/Engine/Source/Editor/FoliageEdit/Private/FoliageEdMode.cpp b/Engine/Source/Editor/FoliageEdit/Private/FoliageEdMode.cpp
--- a/Engine/Source/Editor/FoliageEdit/Private/FoliageEdMode.cpp
+++ b/Engine/Source/Editor/FoliageEdit/Private/FoliageEdMode.cpp
@@ -57,7 +57,7 @@ FEdModeFoliage::FEdModeFoliage()
 		StaticMesh = LoadObject<UStaticMesh>(nullptr, TEXT("/Engine/EngineMeshes/Sphere.Sphere"), nullptr, LOAD_None, nullptr);
 	}
 
-	SphereBrushComponent = ConstructObject<UStaticMeshComponent>(UStaticMeshComponent::StaticClass());
+	SphereBrushComponent = NewObject<UStaticMeshComponent>(GetTransientPackage(), TEXT("SphereBrushComponent"));
 	SphereBrushComponent->SetCollisionProfileName(UCollisionProfile::NoCollision_ProfileName);
 	SphereBrushComponent->SetCollisionObjectType(ECC_WorldDynamic);
 	SphereBrushComponent->StaticMesh = StaticMesh;
diff --git a/Engine/Source/Editor/FontEditor/Private/FontEditor.cpp b/Engine/Source/Editor/FontEditor/Private/FontEditor.cpp
--- a/Engine/Source/Editor/FontEditor/Private/FontEditor.cpp
+++ b/Engine/Source/Editor/FontEditor/Private/FontEditor.cpp
@@ -146,9 +146,9 @@ void FFontEditor::InitFontEditor(const EToolkitMode::Type Mode, const TSharedPtr
 	Font->SetFlags(RF_Transactional);
 	
 	// Create a TGA exporter
-	TGAExporter = ConstructObject<UTextureExporterTGA>(UTextureExporterTGA::StaticClass());
+	TGAExporter = NewObject<UTextureExporterTGA>();
 	// And our importer
-	Factory = ConstructObject<UTextureFactory>(UTextureFactory::StaticClass());
+	Factory = NewObject<UTextureFactory>();
 	// Set the defaults
 	Factory->Blending = BLEND_Opaque;
 	Factory->ShadingModel = MSM_Unlit;
@@ -1012,12 +1012,12 @@ bool FFontEditor::RecreateFontObject(const EFontCacheType NewCacheType)
 	{
 	case EFontCacheType::Offline:
 		// UTrueTypeFontFactory will create a new font object using a texture generated from a user-selection font
-		FontFactoryPtr = ConstructObject<UTrueTypeFontFactory>(UTrueTypeFontFactory::StaticClass());
+		FontFactoryPtr = NewObject<UTrueTypeFontFactory>();
 		break;
 
 	case EFontCacheType::Runtime:
 		// UFontFactory will create an empty font ready to add new font files to
-		FontFactoryPtr = ConstructObject<UFontFactory>(UFontFactory::StaticClass());
+		FontFactoryPtr = NewObject<UFontFactory>();
 		break;
 
 	default:
diff --git a/Engine/Source/Editor/GameProjectGeneration/Private/GameProjectUtils.cpp b/Engine/Source/Editor/GameProjectGeneration/Private/GameProjectUtils.cpp
--- a/Engine/Source/Editor/GameProjectGeneration/Private/GameProjectUtils.cpp
+++ b/Engine/Source/Editor/GameProjectGeneration/Private/GameProjectUtils.cpp
@@ -941,7 +941,7 @@ UTemplateProjectDefs* GameProjectUtils::LoadTemplateDefs(const FString& ProjectD
 				UE_LOG(LogGameProjectGeneration, Error, TEXT("Failed to find template project defs class '%s', using default."), *ClassName);
 			}
 		}
-		TemplateDefs = ConstructObject<UTemplateProjectDefs>(ClassToConstruct);
+		TemplateDefs = NewObject<UTemplateProjectDefs>(GetTransientPackage(), ClassToConstruct);
 		TemplateDefs->LoadConfig(UTemplateProjectDefs::StaticClass(), *TemplateDefsIniFilename);
 	}
 
diff --git a/Engine/Source/Editor/GameplayAbilitiesEditor/Private/AssetTypeActions_GameplayEffect.cpp b/Engine/Source/Editor/GameplayAbilitiesEditor/Private/AssetTypeActions_GameplayEffect.cpp
--- a/Engine/Source/Editor/GameplayAbilitiesEditor/Private/AssetTypeActions_GameplayEffect.cpp
+++ b/Engine/Source/Editor/GameplayAbilitiesEditor/Private/AssetTypeActions_GameplayEffect.cpp
@@ -93,7 +93,7 @@ uint32 FAssetTypeActions_GameplayEffect::GetCategories()
 
 void FAssetTypeActions_GameplayEffect::ExecuteConvertToBlueprint(FWeakGameplayEffectPointerArray Objects)
 {
-	UBlueprintFactory* BlueprintFactory = ConstructObject<UBlueprintFactory>(UBlueprintFactory::StaticClass());
+	UBlueprintFactory* BlueprintFactory = NewObject<UBlueprintFactory>();
 	BlueprintFactory->ParentClass = UGameplayEffect::StaticClass();
 	IAssetTools& AssetTools = FModuleManager::LoadModuleChecked<FAssetToolsModule>("AssetTools").Get();
 
diff --git a/Engine/Source/Editor/GeometryMode/Private/GeometryEdMode.cpp b/Engine/Source/Editor/GeometryMode/Private/GeometryEdMode.cpp
--- a/Engine/Source/Editor/GeometryMode/Private/GeometryEdMode.cpp
+++ b/Engine/Source/Editor/GeometryMode/Private/GeometryEdMode.cpp
@@ -1025,20 +1025,20 @@ FModeTool_GeometryModify::FModeTool_GeometryModify()
 {
 	ID = MT_GeometryModify;
 
-	Modifiers.Add( ConstructObject<UGeomModifier>( UGeomModifier_Edit::StaticClass() ) );
-	Modifiers.Add( ConstructObject<UGeomModifier>( UGeomModifier_Extrude::StaticClass() ) );
-	Modifiers.Add( ConstructObject<UGeomModifier>( UGeomModifier_Clip::StaticClass() ) );
-	Modifiers.Add( ConstructObject<UGeomModifier>( UGeomModifier_Pen::StaticClass() ) );
-	Modifiers.Add( ConstructObject<UGeomModifier>( UGeomModifier_Lathe::StaticClass() ) );
-
-	Modifiers.Add( ConstructObject<UGeomModifier>( UGeomModifier_Create::StaticClass() ) );
-	Modifiers.Add( ConstructObject<UGeomModifier>( UGeomModifier_Delete::StaticClass() ) );
-	Modifiers.Add( ConstructObject<UGeomModifier>( UGeomModifier_Flip::StaticClass() ) );
-	Modifiers.Add( ConstructObject<UGeomModifier>( UGeomModifier_Split::StaticClass() ) );
-	Modifiers.Add( ConstructObject<UGeomModifier>( UGeomModifier_Triangulate::StaticClass() ) );
-	Modifiers.Add( ConstructObject<UGeomModifier>( UGeomModifier_Optimize::StaticClass() ) );
-	Modifiers.Add( ConstructObject<UGeomModifier>( UGeomModifier_Turn::StaticClass() ) );
-	Modifiers.Add( ConstructObject<UGeomModifier>( UGeomModifier_Weld::StaticClass() ) );
+	Modifiers.Add( NewObject<UGeomModifier>( GetTransientPackage(), UGeomModifier_Edit::StaticClass() ) );
+	Modifiers.Add( NewObject<UGeomModifier>( GetTransientPackage(), UGeomModifier_Extrude::StaticClass() ) );
+	Modifiers.Add( NewObject<UGeomModifier>( GetTransientPackage(), UGeomModifier_Clip::StaticClass() ) );
+	Modifiers.Add( NewObject<UGeomModifier>( GetTransientPackage(), UGeomModifier_Pen::StaticClass() ) );
+	Modifiers.Add( NewObject<UGeomModifier>( GetTransientPackage(), UGeomModifier_Lathe::StaticClass() ) );
+
+	Modifiers.Add( NewObject<UGeomModifier>( GetTransientPackage(), UGeomModifier_Create::StaticClass() ) );
+	Modifiers.Add( NewObject<UGeomModifier>( GetTransientPackage(), UGeomModifier_Delete::StaticClass() ) );
+	Modifiers.Add( NewObject<UGeomModifier>( GetTransientPackage(), UGeomModifier_Flip::StaticClass() ) );
+	Modifiers.Add( NewObject<UGeomModifier>( GetTransientPackage(), UGeomModifier_Split::StaticClass() ) );
+	Modifiers.Add( NewObject<UGeomModifier>( GetTransientPackage(), UGeomModifier_Triangulate::StaticClass() ) );
+	Modifiers.Add( NewObject<UGeomModifier>( GetTransientPackage(), UGeomModifier_Optimize::StaticClass() ) );
+	Modifiers.Add( NewObject<UGeomModifier>( GetTransientPackage(), UGeomModifier_Turn::StaticClass() ) );
+	Modifiers.Add( NewObject<UGeomModifier>( GetTransientPackage(), UGeomModifier_Weld::StaticClass() ) );
 
 	CurrentModifier = NULL;
 	
diff --git a/Engine/Source/Editor/GeometryMode/Private/GeometryModifiers.cpp b/Engine/Source/Editor/GeometryMode/Private/GeometryModifiers.cpp
--- a/Engine/Source/Editor/GeometryMode/Private/GeometryModifiers.cpp
+++ b/Engine/Source/Editor/GeometryMode/Private/GeometryModifiers.cpp
@@ -131,7 +131,7 @@ void UGeomModifier::CacheBrushState()
 	if( !CachedPolys )
 	{
 		//Create the list of polys
-		CachedPolys = ConstructObject<UPolys>(UPolys::StaticClass(), this);
+		CachedPolys = NewObject<UPolys>(this);
 	}
 	CachedPolys->Element.Empty();
 
diff --git a/Engine/Source/Editor/Kismet/Private/SCSEditorViewportClient.cpp b/Engine/Source/Editor/Kismet/Private/SCSEditorViewportClient.cpp
--- a/Engine/Source/Editor/Kismet/Private/SCSEditorViewportClient.cpp
+++ b/Engine/Source/Editor/Kismet/Private/SCSEditorViewportClient.cpp
@@ -129,7 +129,7 @@ FSCSEditorViewportClient::FSCSEditorViewportClient(TWeakPtr<FBlueprintEditor>& I
 	DrawHelper.bDrawGrid = GEditor->GetEditorUserSettings().bSCSEditorShowGrid;
 
 	// now add floor
-	EditorFloorComp = ConstructObject<UStaticMeshComponent>(UStaticMeshComponent::StaticClass());
+	EditorFloorComp = NewObject<UStaticMeshComponent>(GetTransientPackage(), TEXT("EditorFloorComp"));
 
 	UStaticMesh* FloorMesh = LoadObject<UStaticMesh>(NULL, TEXT("/Engine/EditorMeshes/PhAT_FloorBox.PhAT_FloorBox"), NULL, LOAD_None, NULL);
 	if (ensure(FloorMesh))
