Added Alembic plugin which currently supports StaticMesh (re)importing from Alembic files (white listed Win64 only for now)
Added Alembic thirdparty dependencies
Changed the FBX CanReimport functionality to allow other ReimportingFactories to assess if they can handle the file first (e.g. SpeedTree and Alembic)

diff --git a/Engine/Plugins/Editor/SpeedTreeImporter/Source/SpeedTreeImporter/Private/ReimportSpeedTreeFactory.cpp b/Engine/Plugins/Editor/SpeedTreeImporter/Source/SpeedTreeImporter/Private/ReimportSpeedTreeFactory.cpp
--- a/Engine/Plugins/Editor/SpeedTreeImporter/Source/SpeedTreeImporter/Private/ReimportSpeedTreeFactory.cpp
+++ b/Engine/Plugins/Editor/SpeedTreeImporter/Source/SpeedTreeImporter/Private/ReimportSpeedTreeFactory.cpp
@@ -29,8 +29,12 @@ bool UReimportSpeedTreeFactory::CanReimport(UObject* Obj, TArray<FString>& OutFi
 	UStaticMesh* Mesh = Cast<UStaticMesh>(Obj);
 	if (Mesh && Mesh->AssetImportData)
 	{
-		Mesh->AssetImportData->ExtractFilenames(OutFilenames);
-		return true;
+		if (FPaths::GetExtension(Mesh->AssetImportData->GetFirstFilename()).ToLower() == "srt")
+		{
+			// SpeedTrees file extension matches with filepath			
+			Mesh->AssetImportData->ExtractFilenames(OutFilenames);
+			return true;
+		}		
 	}
 #endif // #if WITH_SPEEDTREE
 	return false;
@@ -40,7 +44,7 @@ void UReimportSpeedTreeFactory::SetReimportPaths(UObject* Obj, const TArray<FStr
 {
 #if WITH_SPEEDTREE
 	UStaticMesh* Mesh = Cast<UStaticMesh>(Obj);
-	if (Mesh && ensure(NewReimportPaths.Num() == 1))
+	if (Mesh && Mesh->AssetImportData && ensure(NewReimportPaths.Num() == 1))
 	{
 		Mesh->AssetImportData->UpdateFilenameOnly(NewReimportPaths[0]);
 	}
diff --git a/Engine/Source/Editor/UnrealEd/Private/Factories/EditorFactories.cpp b/Engine/Source/Editor/UnrealEd/Private/Factories/EditorFactories.cpp
--- a/Engine/Source/Editor/UnrealEd/Private/Factories/EditorFactories.cpp
+++ b/Engine/Source/Editor/UnrealEd/Private/Factories/EditorFactories.cpp
@@ -5608,6 +5608,9 @@ UReimportFbxStaticMeshFactory::UReimportFbxStaticMeshFactory(const FObjectInitia
 
 	bCreateNew = false;
 	bText = false;
+
+	// Required to allow other StaticMesh re importers to do their CanReimport checks first, and if they fail the FBX will catch it
+	ImportPriority = DefaultImportPriority - 1;
 }
 
 bool UReimportFbxStaticMeshFactory::CanReimport( UObject* Obj, TArray<FString>& OutFilenames )
@@ -5616,17 +5619,8 @@ bool UReimportFbxStaticMeshFactory::CanReimport( UObject* Obj, TArray<FString>&
 	if(Mesh)
 	{
 		if ( Mesh->AssetImportData )
-		{
-			const FString& Filename = Mesh->AssetImportData->GetFirstFilename();
-			if (FPaths::GetExtension(Filename) == "srt")
-			{
-				// SpeedTrees need to use their own importer
-				return false;
-			}
-			else
-			{
-				OutFilenames.Add(Filename);
-			}
+		{		
+			OutFilenames.Add(Mesh->AssetImportData->GetFirstFilename());
 		}
 		else
 		{
