SpeedTree port to Mesh Description, not working anymore with raw mesh

diff --git a/Engine/Plugins/Editor/SpeedTreeImporter/Source/SpeedTreeImporter/Private/SpeedTreeImportFactory.cpp b/Engine/Plugins/Editor/SpeedTreeImporter/Source/SpeedTreeImporter/Private/SpeedTreeImportFactory.cpp
--- a/Engine/Plugins/Editor/SpeedTreeImporter/Source/SpeedTreeImporter/Private/SpeedTreeImportFactory.cpp
+++ b/Engine/Plugins/Editor/SpeedTreeImporter/Source/SpeedTreeImporter/Private/SpeedTreeImportFactory.cpp
@@ -55,9 +55,10 @@
 #include "AssetRegistryModule.h"
 #include "Private/GeomFitUtils.h"
 #include "SpeedTreeWind.h"
-#include "RawMesh.h"
 #include "ComponentReregisterContext.h"
 
+#include "MeshDescription.h"
+
 #if WITH_SPEEDTREE
 
 #ifdef __clang__ // @todo: ThirdParty/SpeedTree/SpeedTreeSDK-v7.0/Include/Core/RenderState_inl.h:52:2: warning: field 'm_eHueVariation' will be initialized after field 'm_eShaderGenerationMode' [-Wreorder]
@@ -914,37 +915,35 @@ static void MakeBodyFromCollisionObjects(UStaticMesh* StaticMesh, const SpeedTre
 	RefreshCollisionChange(*StaticMesh);
 }
 
-
-
-void ProcessTriangleCorner( SpeedTree::CCore& SpeedTree, const int32 TriangleIndex, const int32 Corner, const SpeedTree::SDrawCall* DrawCall, const SpeedTree::st_uint32* Indices32, const SpeedTree::st_uint16* Indices16, FRawMesh& RawMesh, const int32 IndexOffset, const int32 NumUVs, const SpeedTree::SRenderState* RenderState )
+void ProcessTriangleCorner(SpeedTree::CCore& SpeedTree, const int32 TriangleIndex, const int32 Corner, const SpeedTree::SDrawCall* DrawCall, const SpeedTree::st_uint32* Indices32, const SpeedTree::st_uint16* Indices16, UMeshDescription* MeshDescription, const int32 IndexOffset, const int32 NumUVs, const SpeedTree::SRenderState* RenderState)
 {
-	SpeedTree::st_float32 Data[ 4 ];
+	SpeedTree::st_float32 Data[4];
 
 	// flip the triangle winding
 	int32 Index = TriangleIndex * 3 + Corner;
 
-	int32 VertexIndex = DrawCall->m_b32BitIndices ? Indices32[ Index ] : Indices16[ Index ];
-	RawMesh.WedgeIndices.Add( VertexIndex + IndexOffset );
-
+	int32 VertexIndex = DrawCall->m_b32BitIndices ? Indices32[Index] : Indices16[Index];
+	FVertexInstanceID VertexInstanceID = MeshDescription->CreateVertexInstance(FVertexID(VertexIndex + IndexOffset));
+	FMeshVertexInstance& VertexInstance = MeshDescription->GetVertexInstance(VertexInstanceID);
 	// tangents
-	DrawCall->GetProperty( SpeedTree::VERTEX_PROPERTY_NORMAL, VertexIndex, Data );
-	FVector Normal( -Data[ 0 ], Data[ 1 ], Data[ 2 ] );
-	DrawCall->GetProperty( SpeedTree::VERTEX_PROPERTY_TANGENT, VertexIndex, Data );
-	FVector Tangent( -Data[ 0 ], Data[ 1 ], Data[ 2 ] );
-	RawMesh.WedgeTangentX.Add( Tangent );
-	RawMesh.WedgeTangentY.Add( Normal ^ Tangent );
-	RawMesh.WedgeTangentZ.Add( Normal );
+	DrawCall->GetProperty(SpeedTree::VERTEX_PROPERTY_NORMAL, VertexIndex, Data);
+	FVector Normal(-Data[0], Data[1], Data[2]);
+	DrawCall->GetProperty(SpeedTree::VERTEX_PROPERTY_TANGENT, VertexIndex, Data);
+	FVector Tangent(-Data[0], Data[1], Data[2]);
+	VertexInstance.Tangent = Tangent;
+	FVector Binormal = Normal ^ Tangent;
+	VertexInstance.BinormalSign = (FVector::DotProduct(FVector::CrossProduct(Normal, Tangent), Binormal) < 0.0f) ? -1.0f : 1.0f;
+	VertexInstance.Normal = Normal;
 
 	// ao
-	DrawCall->GetProperty( SpeedTree::VERTEX_PROPERTY_AMBIENT_OCCLUSION, VertexIndex, Data );
-	uint8 AO = Data[ 0 ] * 255.0f;
-	RawMesh.WedgeColors.Add( FColor( AO, AO, AO, 255 ) );
+	DrawCall->GetProperty(SpeedTree::VERTEX_PROPERTY_AMBIENT_OCCLUSION, VertexIndex, Data);
+	VertexInstance.Color = FLinearColor(Data[0], Data[0], Data[0]);
 
 	// keep texcoords padded to align indices
-	int32 BaseTexcoordIndex = RawMesh.WedgeTexCoords[ 0 ].Num();
-	for( int32 PadIndex = 0; PadIndex < NumUVs; ++PadIndex )
+	VertexInstance.VertexUVs.Num();
+	for (int32 PadIndex = 0; PadIndex < NumUVs; ++PadIndex)
 	{
-		RawMesh.WedgeTexCoords[ PadIndex ].AddUninitialized( 1 );
+		VertexInstance.VertexUVs.AddUninitialized();
 	}
 
 	// All texcoords are packed into 4 float4 vertex attributes
@@ -963,74 +962,75 @@ void ProcessTriangleCorner( SpeedTree::CCore& SpeedTree, const int32 TriangleInd
 
 
 	// diffuse
-	DrawCall->GetProperty( SpeedTree::VERTEX_PROPERTY_DIFFUSE_TEXCOORDS, VertexIndex, Data );
-	RawMesh.WedgeTexCoords[ 0 ].Top() = FVector2D( Data[ 0 ], Data[ 1 ] );
+	DrawCall->GetProperty(SpeedTree::VERTEX_PROPERTY_DIFFUSE_TEXCOORDS, VertexIndex, Data);
+	VertexInstance.VertexUVs[0] = FVector2D(Data[0], Data[1]);
 
 	// lightmap
-	DrawCall->GetProperty( SpeedTree::VERTEX_PROPERTY_LIGHTMAP_TEXCOORDS, VertexIndex, Data );
-	RawMesh.WedgeTexCoords[ 1 ].Top() = FVector2D( Data[ 0 ], Data[ 1 ] );
+	DrawCall->GetProperty(SpeedTree::VERTEX_PROPERTY_LIGHTMAP_TEXCOORDS, VertexIndex, Data);
+	VertexInstance.VertexUVs[1] = FVector2D(Data[0], Data[1]);
 
 	// branch wind
-	DrawCall->GetProperty( SpeedTree::VERTEX_PROPERTY_WIND_BRANCH_DATA, VertexIndex, Data );
-	RawMesh.WedgeTexCoords[ 2 ].Top() = FVector2D( Data[ 0 ], Data[ 1 ] );
+	DrawCall->GetProperty(SpeedTree::VERTEX_PROPERTY_WIND_BRANCH_DATA, VertexIndex, Data);
+	VertexInstance.VertexUVs[2] = FVector2D(Data[0], Data[1]);
 
 	// lod
-	if( RenderState->m_bFacingLeavesPresent )
+	if (RenderState->m_bFacingLeavesPresent)
 	{
-		DrawCall->GetProperty( SpeedTree::VERTEX_PROPERTY_LEAF_CARD_LOD_SCALAR, VertexIndex, Data );
-		RawMesh.WedgeTexCoords[ 3 ].Top() = FVector2D( Data[ 0 ], 0.0f );
-		RawMesh.WedgeTexCoords[ 4 ].Top() = FVector2D( 0.0f, 0.0f );
+		DrawCall->GetProperty(SpeedTree::VERTEX_PROPERTY_LEAF_CARD_LOD_SCALAR, VertexIndex, Data);
+		VertexInstance.VertexUVs[3] = FVector2D(Data[0], 0.0f);
+		VertexInstance.VertexUVs[4] = FVector2D(0.0f, 0.0f);
 	}
 	else
 	{
-		DrawCall->GetProperty( SpeedTree::VERTEX_PROPERTY_LOD_POSITION, VertexIndex, Data );
-		RawMesh.WedgeTexCoords[ 3 ].Top() = FVector2D( -Data[ 0 ], Data[ 1 ] );
-		RawMesh.WedgeTexCoords[ 4 ].Top() = FVector2D( Data[ 2 ], 0.0f );
+		DrawCall->GetProperty(SpeedTree::VERTEX_PROPERTY_LOD_POSITION, VertexIndex, Data);
+		VertexInstance.VertexUVs[3] = FVector2D(-Data[0], Data[1]);
+		VertexInstance.VertexUVs[4] = FVector2D(Data[2], 0.0f);
 	}
 
 	// other
-	if( RenderState->m_bBranchesPresent )
+	if (RenderState->m_bBranchesPresent)
 	{
 		// detail
-		DrawCall->GetProperty( SpeedTree::VERTEX_PROPERTY_DETAIL_TEXCOORDS, VertexIndex, Data );
-		RawMesh.WedgeTexCoords[ 5 ].Top() = FVector2D( Data[ 0 ], Data[ 1 ] );
+		DrawCall->GetProperty(SpeedTree::VERTEX_PROPERTY_DETAIL_TEXCOORDS, VertexIndex, Data);
+		VertexInstance.VertexUVs[5] = FVector2D(Data[0], Data[1]);
 
 		// branch seam
-		DrawCall->GetProperty( SpeedTree::VERTEX_PROPERTY_BRANCH_SEAM_DIFFUSE, VertexIndex, Data );
-		RawMesh.WedgeTexCoords[ 6 ].Top() = FVector2D( Data[ 0 ], Data[ 1 ] );
-		RawMesh.WedgeTexCoords[ 4 ].Top().Y = Data[ 2 ];
+		DrawCall->GetProperty(SpeedTree::VERTEX_PROPERTY_BRANCH_SEAM_DIFFUSE, VertexIndex, Data);
+		VertexInstance.VertexUVs[6] = FVector2D(Data[0], Data[1]);
+		VertexInstance.VertexUVs[4].Y = Data[2];
 	}
-	else if( RenderState->m_bFrondsPresent )
+	else if (RenderState->m_bFrondsPresent)
 	{
 		// frond wind
-		DrawCall->GetProperty( SpeedTree::VERTEX_PROPERTY_WIND_EXTRA_DATA, VertexIndex, Data );
-		RawMesh.WedgeTexCoords[ 5 ].Top() = FVector2D( Data[ 0 ], Data[ 1 ] );
-		RawMesh.WedgeTexCoords[ 6 ].Top() = FVector2D( Data[ 2 ], 0.0f );
+		DrawCall->GetProperty(SpeedTree::VERTEX_PROPERTY_WIND_EXTRA_DATA, VertexIndex, Data);
+		VertexInstance.VertexUVs[5] = FVector2D(Data[0], Data[1]);
+		VertexInstance.VertexUVs[6] = FVector2D(Data[2], 0.0f);
 	}
-	else if( RenderState->m_bLeavesPresent || RenderState->m_bFacingLeavesPresent )
+	else if (RenderState->m_bLeavesPresent || RenderState->m_bFacingLeavesPresent)
 	{
 		// anchor
-		if( RenderState->m_bFacingLeavesPresent )
+		if (RenderState->m_bFacingLeavesPresent)
 		{
-			DrawCall->GetProperty( SpeedTree::VERTEX_PROPERTY_POSITION, VertexIndex, Data );
+			DrawCall->GetProperty(SpeedTree::VERTEX_PROPERTY_POSITION, VertexIndex, Data);
 		}
 		else
 		{
-			DrawCall->GetProperty( SpeedTree::VERTEX_PROPERTY_LEAF_ANCHOR_POINT, VertexIndex, Data );
+			DrawCall->GetProperty(SpeedTree::VERTEX_PROPERTY_LEAF_ANCHOR_POINT, VertexIndex, Data);
 		}
-		RawMesh.WedgeTexCoords[ 4 ].Top().Y = -Data[ 0 ];
-		RawMesh.WedgeTexCoords[ 5 ].Top() = FVector2D( Data[ 1 ], Data[ 2 ] );
+		VertexInstance.VertexUVs[4].Y = -Data[0];
+		VertexInstance.VertexUVs[5] = FVector2D(Data[1], Data[2]);
 
 		// leaf wind
-		DrawCall->GetProperty( SpeedTree::VERTEX_PROPERTY_WIND_EXTRA_DATA, VertexIndex, Data );
-		RawMesh.WedgeTexCoords[ 6 ].Top() = FVector2D( Data[ 0 ], Data[ 1 ] );
-		RawMesh.WedgeTexCoords[ 7 ].Top().X = Data[ 2 ];
-		DrawCall->GetProperty( SpeedTree::VERTEX_PROPERTY_WIND_FLAGS, VertexIndex, Data );
-		RawMesh.WedgeTexCoords[ 7 ].Top().Y = Data[ 0 ];
+		DrawCall->GetProperty(SpeedTree::VERTEX_PROPERTY_WIND_EXTRA_DATA, VertexIndex, Data);
+		VertexInstance.VertexUVs[6] = FVector2D(Data[0], Data[1]);
+		VertexInstance.VertexUVs[7].X = Data[2];
+		DrawCall->GetProperty(SpeedTree::VERTEX_PROPERTY_WIND_FLAGS, VertexIndex, Data);
+		VertexInstance.VertexUVs[7].Y = Data[0];
 	}
 }
 
 
+
 UObject* USpeedTreeImportFactory::FactoryCreateBinary(UClass* InClass, UObject* InParent, FName InName, EObjectFlags Flags, UObject* Context, const TCHAR* Type, const uint8*& Buffer, const uint8* BufferEnd, FFeedbackContext* Warn, bool& bOutOperationCanceled)
 {
 	FEditorDelegates::OnAssetPreImport.Broadcast(this, InClass, InParent, InName, Type);
@@ -1196,7 +1196,18 @@ UObject* USpeedTreeImportFactory::FactoryCreateBinary(UClass* InClass, UObject*
 					for (int32 LODIndex = 0; LODIndex < SpeedTreeGeometry->m_nNumLods; ++LODIndex)
 					{
 						const SpeedTree::SLod* TreeLOD = &SpeedTreeGeometry->m_pLods[LODIndex];
-						FRawMesh RawMesh;
+						
+						FStaticMeshSourceModel* LODModel = new (StaticMesh->SourceModels) FStaticMeshSourceModel();
+						LODModel->BuildSettings.bRecomputeNormals = false;
+						LODModel->BuildSettings.bRecomputeTangents = false;
+						LODModel->BuildSettings.bRemoveDegenerates = true;
+						LODModel->BuildSettings.bUseHighPrecisionTangentBasis = false;
+						LODModel->BuildSettings.bUseFullPrecisionUVs = false;
+						LODModel->BuildSettings.bGenerateLightmapUVs = false;
+						LODModel->ScreenSize = 0.1f / FMath::Pow(2.0f, StaticMesh->SourceModels.Num() - 1);
+
+						UMeshDescription* MeshDescription = NewObject<UMeshDescription>(Package);
+						StaticMesh->SetOriginalMeshDescription(LODIndex, MeshDescription);
 
 						// compute the number of texcoords we need so we can pad when necessary
 						int32 NumUVs = 7; // static meshes have fewer, but they are so rare, we shouldn't complicate things for them
@@ -1269,15 +1280,28 @@ UObject* USpeedTreeImportFactory::FactoryCreateBinary(UClass* InClass, UObject*
 								
 								RenderStateIndexToStaticMeshIndex.Add(DrawCall->m_nRenderStateIndex, StaticMesh->StaticMaterials.Num());
 								MaterialIndex = StaticMesh->StaticMaterials.Num();
-								StaticMesh->StaticMaterials.Add(FStaticMaterial(Material));
+								StaticMesh->StaticMaterials.Add(FStaticMaterial(Material, *MaterialName, *MaterialName));
 							}
 							else
 							{
 								MaterialIndex = *OldMaterial;
 							}
 
-							int32 IndexOffset = RawMesh.VertexPositions.Num();
+							//Add the polygon group
+							while (MaterialIndex >= MeshDescription->PolygonGroups().Num())
+							{
+								const FStaticMaterial& Material = StaticMesh->StaticMaterials[MeshDescription->PolygonGroups().Num()];
+								FPolygonGroupID PolygonGroupID = MeshDescription->CreatePolygonGroup();
+								FMeshPolygonGroup& NewPolygonGroup = MeshDescription->GetPolygonGroup(PolygonGroupID);
+								NewPolygonGroup.bCastShadow = true;
+								NewPolygonGroup.bEnableCollision = true;
+								NewPolygonGroup.ImportedMaterialSlotName = Material.ImportedMaterialSlotName;
+								NewPolygonGroup.MaterialSlotName = Material.MaterialSlotName;
+								NewPolygonGroup.MaterialAsset = FStringAssetReference(Material.MaterialInterface);
+							}
 
+							int32 IndexOffset = MeshDescription->Vertices().Num();
+							//Create the vertices
 							for (int32 VertexIndex = 0; VertexIndex < DrawCall->m_nNumVertices; ++VertexIndex)
 							{
 								// position
@@ -1291,8 +1315,9 @@ UObject* USpeedTreeImportFactory::FactoryCreateBinary(UClass* InClass, UObject*
 									Data[1] += Data2[1];
 									Data[2] += Data2[2];                                    
 								}
-
-								RawMesh.VertexPositions.Add(FVector(-Data[0], Data[1], Data[2]));
+								FVertexID VertexID = MeshDescription->CreateVertex();
+								FMeshVertex& Vertex = MeshDescription->GetVertex(VertexID);
+								Vertex.VertexPosition = FVector(-Data[0], Data[1], Data[2]);
 							}
 
 							const SpeedTree::st_byte* pIndexData = &*DrawCall->m_pIndexData;
@@ -1300,63 +1325,78 @@ UObject* USpeedTreeImportFactory::FactoryCreateBinary(UClass* InClass, UObject*
 							const SpeedTree::st_uint16* Indices16 = (SpeedTree::st_uint16*)pIndexData;
 
 							int32 TriangleCount = DrawCall->m_nNumIndices / 3;
-							int32 ExistingTris = RawMesh.WedgeIndices.Num() / 3;
+							int32 ExistingTris = MeshDescription->Polygons().Num();// RawMesh.WedgeIndices.Num() / 3;
 
 							for (int32 TriangleIndex = 0; TriangleIndex < TriangleCount; ++TriangleIndex)
 							{
-								RawMesh.FaceMaterialIndices.Add(MaterialIndex);
-								RawMesh.FaceSmoothingMasks.Add(0);
-
-#if 0
-								float LeafLODScalar = 1.0f;
-								if (RenderState->m_bLeavesPresent)
+								FPolygonGroupID TrianglePolygonGroupID(MaterialIndex);
+								FMeshPolygonGroup& PolygonGroup = MeshDescription->GetPolygonGroup(TrianglePolygonGroupID);
+								FVertexInstanceID CornerInstanceIDs[3];
+								FVertexID CornerVerticesIDs[3];
+								//Create the vertex instances
+								for (int32 Corner = 0; Corner < 3; ++Corner)
 								{
-									// reconstruct a single leaf scalar so we don't need the lod position
-									// This is not 100% like the Modeler for frond-based leaves, but it reduces the amount of VB data needed
-									int32 Index = TriangleIndex * 3;
-									int32 VertexIndex1 = DrawCall->m_b32BitIndices ? Indices32[Index + 0] : Indices16[Index + 0];
-									int32 VertexIndex2 = DrawCall->m_b32BitIndices ? Indices32[Index + 1] : Indices16[Index + 1];
-									int32 VertexIndex3 = DrawCall->m_b32BitIndices ? Indices32[Index + 2] : Indices16[Index + 2];
-
-									SpeedTree::st_float32 Data[4];
-									DrawCall->GetProperty(SpeedTree::VERTEX_PROPERTY_POSITION, VertexIndex1, Data);
-									FVector Corner1(Data[0], Data[1], Data[2]);
-									DrawCall->GetProperty(SpeedTree::VERTEX_PROPERTY_POSITION, VertexIndex2, Data);
-									FVector Corner2(Data[0], Data[1], Data[2]);
-									DrawCall->GetProperty(SpeedTree::VERTEX_PROPERTY_POSITION, VertexIndex3, Data);
-									FVector Corner3(Data[0], Data[1], Data[2]);
-
-									DrawCall->GetProperty(SpeedTree::VERTEX_PROPERTY_LOD_POSITION, VertexIndex1, Data);
-									FVector LODCorner1(Data[0], Data[1], Data[2]);
-									DrawCall->GetProperty(SpeedTree::VERTEX_PROPERTY_LOD_POSITION, VertexIndex2, Data);
-									FVector LODCorner2(Data[0], Data[1], Data[2]);
-									DrawCall->GetProperty(SpeedTree::VERTEX_PROPERTY_LOD_POSITION, VertexIndex3, Data);
-									FVector LODCorner3(Data[0], Data[1], Data[2]);
-
-									float OrigArea = ((Corner2 - Corner1) ^ (Corner3 - Corner1)).Size();
-									float LODArea = ((LODCorner2 - LODCorner1) ^ (LODCorner3 - LODCorner1)).Size();
-
-									LeafLODScalar = FMath::Sqrt(LODArea / OrigArea);
+									ProcessTriangleCorner( SpeedTree, TriangleIndex, Corner, DrawCall, Indices32, Indices16, MeshDescription, IndexOffset, NumUVs, RenderState );
+									int32 VertexInstanceIndex = ((ExistingTris + TriangleIndex) * 3) + Corner;
+									const FVertexInstanceID VertexInstanceID(VertexInstanceIndex);
+									CornerInstanceIDs[Corner] = VertexInstanceID;
+									const FMeshVertexInstance& TriangleVertexInstance = MeshDescription->GetVertexInstance(VertexInstanceID);
+									CornerVerticesIDs[Corner] = TriangleVertexInstance.VertexID;
+								}
+								// Create polygon edges
+								TArray<UMeshDescription::FContourPoint> Contours;
+								// Add the edges of this triangle
+								for (uint32 TriangleEdgeNumber = 0; TriangleEdgeNumber < 3; ++TriangleEdgeNumber)
+								{
+									int32 ContourPointIndex = Contours.AddDefaulted();
+									UMeshDescription::FContourPoint& ContourPoint = Contours[ContourPointIndex];
+									//Find the matching edge ID
+									uint32 CornerIndices[2];
+									CornerIndices[0] = (TriangleEdgeNumber + 0) % 3;
+									CornerIndices[1] = (TriangleEdgeNumber + 1) % 3;
+
+									FVertexID EdgeVertexIDs[2];
+									EdgeVertexIDs[0] = CornerVerticesIDs[CornerIndices[0]];
+									EdgeVertexIDs[1] = CornerVerticesIDs[CornerIndices[1]];
+
+									FEdgeID MatchEdgeId = FEdgeID::Invalid;
+									const FMeshVertex& EdgeMeshVertex = MeshDescription->GetVertex(EdgeVertexIDs[0]);
+									for (const FEdgeID& EdgeID : EdgeMeshVertex.ConnectedEdgeIDs)
+									{
+										FMeshEdge& Edge = MeshDescription->GetEdge(EdgeID);
+										if (EdgeVertexIDs[0] == Edge.VertexIDs[0] && EdgeVertexIDs[1] == Edge.VertexIDs[1] ||
+											EdgeVertexIDs[1] == Edge.VertexIDs[0] && EdgeVertexIDs[0] == Edge.VertexIDs[1])
+										{
+											MatchEdgeId = EdgeID;
+											break;
+										}
+									}
+									if (MatchEdgeId == FEdgeID::Invalid)
+									{
+										MatchEdgeId = MeshDescription->CreateEdge(EdgeVertexIDs[0], EdgeVertexIDs[1]);
+									}
+									ContourPoint.EdgeID = MatchEdgeId;
+									ContourPoint.VertexInstanceID = CornerInstanceIDs[CornerIndices[0]];
+									FMeshEdge& ExistingEdge = MeshDescription->GetEdge(MatchEdgeId);
+									ExistingEdge.CreaseSharpness = 0.0f;
+									ExistingEdge.bIsHardEdge = false;
 								}
-#endif
 
-								for (int32 Corner = 0; Corner < 3; ++Corner)
+								// Insert a polygon into the mesh
+								const FPolygonID NewPolygonID = MeshDescription->CreatePolygon(Contours);
+								FMeshPolygon& NewPolygon = MeshDescription->GetPolygon(NewPolygonID);
+								NewPolygon.PolygonGroupID = TrianglePolygonGroupID;
+								PolygonGroup.Polygons.Add(NewPolygonID);
+								int32 NewTriangleIndex = NewPolygon.Triangles.AddDefaulted();
+								FMeshTriangle& NewTriangle = NewPolygon.Triangles[NewTriangleIndex];
+								for (int32 TriangleVertexIndex = 0; TriangleVertexIndex < 3; ++TriangleVertexIndex)
 								{
-									ProcessTriangleCorner( SpeedTree, TriangleIndex, Corner, DrawCall, Indices32, Indices16, RawMesh, IndexOffset, NumUVs, RenderState );
+									const FVertexInstanceID &VertexInstanceID = CornerInstanceIDs[TriangleVertexIndex];
+									NewTriangle.SetVertexInstanceID(TriangleVertexIndex, VertexInstanceID);
 								}
 							}							
 						}
 
-						FStaticMeshSourceModel* LODModel = new (StaticMesh->SourceModels) FStaticMeshSourceModel();
-						LODModel->BuildSettings.bRecomputeNormals = false;
-						LODModel->BuildSettings.bRecomputeTangents = false;
-						LODModel->BuildSettings.bRemoveDegenerates = true;
-						LODModel->BuildSettings.bUseHighPrecisionTangentBasis = false;
-						LODModel->BuildSettings.bUseFullPrecisionUVs = false;
-						LODModel->BuildSettings.bGenerateLightmapUVs = false;
-						LODModel->ScreenSize = 0.1f / FMath::Pow(2.0f, StaticMesh->SourceModels.Num() - 1);
-						LODModel->RawMeshBulkData->SaveRawMesh(RawMesh);
-
 						for (int32 MaterialIndex = 0; MaterialIndex < StaticMesh->StaticMaterials.Num(); ++MaterialIndex)
 						{
 							FMeshSectionInfo Info = StaticMesh->SectionInfoMap.Get(LODIndex, MaterialIndex);
@@ -1369,12 +1409,41 @@ UObject* USpeedTreeImportFactory::FactoryCreateBinary(UClass* InClass, UObject*
 				// make billboard LOD
 				if (Options->SpeedTreeImportData->ImportGeometryType != EImportGeometryType::IGT_3D && SpeedTreeGeometry->m_sVertBBs.m_nNumBillboards > 0)
 				{
-					UMaterialInterface* Material = CreateSpeedTreeMaterial(InParent, MeshName + "_Billboard", &SpeedTreeGeometry->m_aBillboardRenderStates[SpeedTree::RENDER_PASS_MAIN], Options, WindType, SpeedTreeGeometry->m_sVertBBs.m_nNumBillboards, LoadedPackages);
+					const int32 LODIndex = StaticMesh->SourceModels.Num();
+					FString MaterialFullName = MeshName + TEXT("_Billboard");
+					UMaterialInterface* Material = CreateSpeedTreeMaterial(InParent, MaterialFullName, &SpeedTreeGeometry->m_aBillboardRenderStates[SpeedTree::RENDER_PASS_MAIN], Options, WindType, SpeedTreeGeometry->m_sVertBBs.m_nNumBillboards, LoadedPackages);
 					int32 MaterialIndex = StaticMesh->StaticMaterials.Num();
-					StaticMesh->StaticMaterials.Add(FStaticMaterial(Material));
+					StaticMesh->StaticMaterials.Add(FStaticMaterial(Material, *MaterialFullName, *MaterialFullName));
+
+					FStaticMeshSourceModel* LODModel = new (StaticMesh->SourceModels) FStaticMeshSourceModel();
+					LODModel->BuildSettings.bRecomputeNormals = false;
+					LODModel->BuildSettings.bRecomputeTangents = false;
+					LODModel->BuildSettings.bRemoveDegenerates = true;
+					LODModel->BuildSettings.bUseHighPrecisionTangentBasis = false;
+					LODModel->BuildSettings.bUseFullPrecisionUVs = false;
+					LODModel->BuildSettings.bGenerateLightmapUVs = false;
+					LODModel->ScreenSize = 0.1f / FMath::Pow(2.0f, StaticMesh->SourceModels.Num() - 1);
+					// Add mesh section info entry for billboard LOD (only one section/material index)
+					FMeshSectionInfo Info = StaticMesh->SectionInfoMap.Get(LODIndex, 0);
+					Info.MaterialIndex = MaterialIndex;
+					StaticMesh->SectionInfoMap.Set(LODIndex, 0, Info);
 
-					FRawMesh RawMesh;
+					UMeshDescription* MeshDescription = NewObject<UMeshDescription>(Package);
+					StaticMesh->SetOriginalMeshDescription(LODIndex, MeshDescription);
 					
+					//Add the polygon group
+					while (MaterialIndex >= MeshDescription->PolygonGroups().Num())
+					{
+						const FStaticMaterial& StaticMaterial = StaticMesh->StaticMaterials[MeshDescription->PolygonGroups().Num()];
+						FPolygonGroupID PolygonGroupID = MeshDescription->CreatePolygonGroup();
+						FMeshPolygonGroup& NewPolygonGroup = MeshDescription->GetPolygonGroup(PolygonGroupID);
+						NewPolygonGroup.bCastShadow = true;
+						NewPolygonGroup.bEnableCollision = true;
+						NewPolygonGroup.ImportedMaterialSlotName = StaticMaterial.ImportedMaterialSlotName;
+						NewPolygonGroup.MaterialSlotName = StaticMaterial.MaterialSlotName;
+						NewPolygonGroup.MaterialAsset = FStringAssetReference(StaticMaterial.MaterialInterface);
+					}
+
 					// fill out triangles
 					float BillboardWidth = SpeedTreeGeometry->m_sVertBBs.m_fWidth;
 					float BillboardBottom = SpeedTreeGeometry->m_sVertBBs.m_fBottomPos;
@@ -1417,63 +1486,113 @@ UObject* USpeedTreeImportFactory::FactoryCreateBinary(UClass* InClass, UObject*
 						const float* TexCoords = &SpeedTreeGeometry->m_sVertBBs.m_pTexCoords[BillboardIndex * 4];
 						bool bRotated = (SpeedTreeGeometry->m_sVertBBs.m_pRotated[BillboardIndex] == 1);
 
-						int32 IndexOffset = RawMesh.VertexPositions.Num();
+						int32 IndexOffset = MeshDescription->Vertices().Num();
 					
-						// position
+						// Vertices
 						for (int32 VertexIndex = 0; VertexIndex < NumVertices; ++VertexIndex)
 						{
 							const SpeedTree::st_float32* Vertex = &Vertices[VertexIndex * 2];
 							FVector Position = BillboardRotate.TransformVector(FVector(Vertex[0] * BillboardWidth - BillboardWidth * 0.5f, 0.0f, Vertex[1] * BillboardHeight + BillboardBottom));
-							RawMesh.VertexPositions.Add(Position);
+							
+							FVertexID VertexID = MeshDescription->CreateVertex();
+							check(VertexID.GetValue() == (VertexIndex + IndexOffset));
+							FMeshVertex& MeshVertex = MeshDescription->GetVertex(VertexID);
+							MeshVertex.VertexPosition = Position;
 						}
 
 						// other data
 						int32 NumTriangles = NumIndices / 3;
 						for (int32 TriangleIndex = 0; TriangleIndex < NumTriangles; ++TriangleIndex)
 						{
-							RawMesh.FaceMaterialIndices.Add(MaterialIndex);
-							RawMesh.FaceSmoothingMasks.Add(0);
-
+							FPolygonGroupID TrianglePolygonGroupID(MaterialIndex);
+							FMeshPolygonGroup& PolygonGroup = MeshDescription->GetPolygonGroup(TrianglePolygonGroupID);
+							FVertexInstanceID CornerInstanceIDs[3];
+							FVertexID CornerVerticesIDs[3];
+							//Create the vertex instances
 							for (int32 Corner = 0; Corner < 3; ++Corner)
 							{
-								int32 Index = Indices[TriangleIndex * 3 + Corner];
+								int32 VertexInstanceIndex = TriangleIndex * 3 + Corner;
+								int32 Index = Indices[VertexInstanceIndex];
 								const SpeedTree::st_float32* Vertex = &Vertices[Index * 2];
 
-								RawMesh.WedgeIndices.Add(Index + IndexOffset);
-
-								RawMesh.WedgeTangentX.Add(TangentX);
-								RawMesh.WedgeTangentY.Add(TangentY);
-								RawMesh.WedgeTangentZ.Add(TangentZ);
+								FVertexInstanceID VertexInstanceID = MeshDescription->CreateVertexInstance(FVertexID(Index + IndexOffset));
+								FMeshVertexInstance& VertexInstance = MeshDescription->GetVertexInstance(VertexInstanceID);
+								VertexInstance.Tangent = TangentX;
+								VertexInstance.BinormalSign = (FVector::DotProduct(FVector::CrossProduct(TangentZ, TangentX), TangentY) < 0.0f) ? -1.0f : 1.0f;
+								VertexInstance.Normal = TangentZ;
 
+								FVector2D VertexUV;
 								if (bRotated)
 								{
-									RawMesh.WedgeTexCoords[0].Add(FVector2D(TexCoords[0] + Vertex[1] * TexCoords[2], TexCoords[1] + Vertex[0] * TexCoords[3]));
+									VertexUV = FVector2D(TexCoords[0] + Vertex[1] * TexCoords[2], TexCoords[1] + Vertex[0] * TexCoords[3]);
 								}
 								else
 								{
-									RawMesh.WedgeTexCoords[0].Add(FVector2D(TexCoords[0] + Vertex[0] * TexCoords[2], TexCoords[1] + Vertex[1] * TexCoords[3]));
+									VertexUV = FVector2D(TexCoords[0] + Vertex[0] * TexCoords[2], TexCoords[1] + Vertex[1] * TexCoords[3]);
 								}
 
-								// lightmap coord
-								RawMesh.WedgeTexCoords[1].Add(RawMesh.WedgeTexCoords[0].Top());
+								// UV Channel 0
+								VertexInstance.VertexUVs.Add(VertexUV);
+								// lightmap coord (UV Channel 1)
+								VertexInstance.VertexUVs.Add(VertexUV);
+								
+								CornerInstanceIDs[Corner] = VertexInstanceID;
+								const FMeshVertexInstance& TriangleVertexInstance = MeshDescription->GetVertexInstance(VertexInstanceID);
+								CornerVerticesIDs[Corner] = TriangleVertexInstance.VertexID;
+							}
+							// Create polygon edges
+							TArray<UMeshDescription::FContourPoint> Contours;
+							// Add the edges of this triangle
+							for (uint32 TriangleEdgeNumber = 0; TriangleEdgeNumber < 3; ++TriangleEdgeNumber)
+							{
+								int32 ContourPointIndex = Contours.AddDefaulted();
+								UMeshDescription::FContourPoint& ContourPoint = Contours[ContourPointIndex];
+								//Find the matching edge ID
+								uint32 CornerIndices[2];
+								CornerIndices[0] = (TriangleEdgeNumber + 0) % 3;
+								CornerIndices[1] = (TriangleEdgeNumber + 1) % 3;
+
+								FVertexID EdgeVertexIDs[2];
+								EdgeVertexIDs[0] = CornerVerticesIDs[CornerIndices[0]];
+								EdgeVertexIDs[1] = CornerVerticesIDs[CornerIndices[1]];
+
+								FEdgeID MatchEdgeId = FEdgeID::Invalid;
+								const FMeshVertex& EdgeMeshVertex = MeshDescription->GetVertex(EdgeVertexIDs[0]);
+								for (const FEdgeID& EdgeID : EdgeMeshVertex.ConnectedEdgeIDs)
+								{
+									FMeshEdge& Edge = MeshDescription->GetEdge(EdgeID);
+									if (EdgeVertexIDs[0] == Edge.VertexIDs[0] && EdgeVertexIDs[1] == Edge.VertexIDs[1] ||
+										EdgeVertexIDs[1] == Edge.VertexIDs[0] && EdgeVertexIDs[0] == Edge.VertexIDs[1])
+									{
+										MatchEdgeId = EdgeID;
+										break;
+									}
+								}
+								if (MatchEdgeId == FEdgeID::Invalid)
+								{
+									MatchEdgeId = MeshDescription->CreateEdge(EdgeVertexIDs[0], EdgeVertexIDs[1]);
+								}
+								ContourPoint.EdgeID = MatchEdgeId;
+								ContourPoint.VertexInstanceID = CornerInstanceIDs[CornerIndices[0]];
+								FMeshEdge& ExistingEdge = MeshDescription->GetEdge(MatchEdgeId);
+								ExistingEdge.CreaseSharpness = 0.0f;
+								ExistingEdge.bIsHardEdge = false;
+							}
+
+							// Insert a polygon into the mesh
+							const FPolygonID NewPolygonID = MeshDescription->CreatePolygon(Contours);
+							FMeshPolygon& NewPolygon = MeshDescription->GetPolygon(NewPolygonID);
+							NewPolygon.PolygonGroupID = TrianglePolygonGroupID;
+							PolygonGroup.Polygons.Add(NewPolygonID);
+							int32 NewTriangleIndex = NewPolygon.Triangles.AddDefaulted();
+							FMeshTriangle& NewTriangle = NewPolygon.Triangles[NewTriangleIndex];
+							for (int32 TriangleVertexIndex = 0; TriangleVertexIndex < 3; ++TriangleVertexIndex)
+							{
+								const FVertexInstanceID &VertexInstanceID = CornerInstanceIDs[TriangleVertexIndex];
+								NewTriangle.SetVertexInstanceID(TriangleVertexIndex, VertexInstanceID);
 							}
 						}
 					}
-
-					FStaticMeshSourceModel* LODModel = new (StaticMesh->SourceModels) FStaticMeshSourceModel();
-					LODModel->BuildSettings.bRecomputeNormals = false;
-					LODModel->BuildSettings.bRecomputeTangents = false;
-					LODModel->BuildSettings.bRemoveDegenerates = true;
-					LODModel->BuildSettings.bUseHighPrecisionTangentBasis = false;
-					LODModel->BuildSettings.bUseFullPrecisionUVs = false;
-					LODModel->BuildSettings.bGenerateLightmapUVs = false;
-					LODModel->ScreenSize = 0.1f / FMath::Pow(2.0f, StaticMesh->SourceModels.Num() - 1);
-					LODModel->RawMeshBulkData->SaveRawMesh(RawMesh);
-					// Add mesh section info entry for billboard LOD (only one section/material index)
-					const int32 LODIndex = StaticMesh->SourceModels.Num() - 1;
-					FMeshSectionInfo Info = StaticMesh->SectionInfoMap.Get(LODIndex, 0);
-					Info.MaterialIndex = MaterialIndex;
-					StaticMesh->SectionInfoMap.Set(LODIndex, 0, Info);
 				}
 
 				if (OldMaterials.Num() == StaticMesh->StaticMaterials.Num())
diff --git a/Engine/Plugins/Editor/SpeedTreeImporter/Source/SpeedTreeImporter/SpeedTreeImporter.Build.cs b/Engine/Plugins/Editor/SpeedTreeImporter/Source/SpeedTreeImporter/SpeedTreeImporter.Build.cs
--- a/Engine/Plugins/Editor/SpeedTreeImporter/Source/SpeedTreeImporter/SpeedTreeImporter.Build.cs
+++ b/Engine/Plugins/Editor/SpeedTreeImporter/Source/SpeedTreeImporter/SpeedTreeImporter.Build.cs
@@ -36,7 +36,8 @@ namespace UnrealBuildTool.Rules
                     "InputCore",
 				    "RawMesh",
                     "UnrealEd",
-                    "MainFrame"
+                    "MainFrame",
+                    "MeshDescription"
 				}
 				);
 
