Fix speedtree importer, material assignation was not uptodate with the new material slot name we use instead of the index

diff --git a/Engine/Plugins/Editor/SpeedTreeImporter/Source/SpeedTreeImporter/Private/SpeedTreeImportFactory.cpp b/Engine/Plugins/Editor/SpeedTreeImporter/Source/SpeedTreeImporter/Private/SpeedTreeImportFactory.cpp
--- a/Engine/Plugins/Editor/SpeedTreeImporter/Source/SpeedTreeImporter/Private/SpeedTreeImportFactory.cpp
+++ b/Engine/Plugins/Editor/SpeedTreeImporter/Source/SpeedTreeImporter/Private/SpeedTreeImportFactory.cpp
@@ -1705,7 +1705,7 @@ UObject* USpeedTreeImportFactory::FactoryCreateBinary7(UClass* InClass, UObject*
 								
 								RenderStateIndexToStaticMeshIndex.Add(DrawCall->m_nRenderStateIndex, StaticMesh->StaticMaterials.Num());
 								MaterialIndex = StaticMesh->StaticMaterials.Num();
-								StaticMesh->StaticMaterials.Add(FStaticMaterial(Material));
+								StaticMesh->StaticMaterials.Add(FStaticMaterial(Material, FName(*MaterialName), FName(*MaterialName)));
 							}
 							else
 							{
@@ -1806,9 +1806,10 @@ UObject* USpeedTreeImportFactory::FactoryCreateBinary7(UClass* InClass, UObject*
 				// make billboard LOD
 				if (Options->SpeedTreeImportData->ImportGeometryType != EImportGeometryType::IGT_3D && SpeedTreeGeometry->m_sVertBBs.m_nNumBillboards > 0)
 				{
-					UMaterialInterface* Material = CreateSpeedTreeMaterial7(InParent, MeshName + "_Billboard", &SpeedTreeGeometry->m_aBillboardRenderStates[SpeedTree::RENDER_PASS_MAIN], Options, WindType, SpeedTreeGeometry->m_sVertBBs.m_nNumBillboards, LoadedPackages);
+					FString MaterialName = MeshName + "_Billboard";
+					UMaterialInterface* Material = CreateSpeedTreeMaterial7(InParent, MaterialName, &SpeedTreeGeometry->m_aBillboardRenderStates[SpeedTree::RENDER_PASS_MAIN], Options, WindType, SpeedTreeGeometry->m_sVertBBs.m_nNumBillboards, LoadedPackages);
 					int32 MaterialIndex = StaticMesh->StaticMaterials.Num();
-					StaticMesh->StaticMaterials.Add(FStaticMaterial(Material));
+					StaticMesh->StaticMaterials.Add(FStaticMaterial(Material, FName(*MaterialName), FName(*MaterialName)));
 
 					FRawMesh RawMesh;
 					
@@ -2140,7 +2141,7 @@ UObject* USpeedTreeImportFactory::FactoryCreateBinary8(UClass* InClass, UObject*
 					MaterialName.InsertAt(MaterialName.Len() - 4, GeomString);
 					UMaterialInterface* Material = CreateSpeedTreeMaterial8(InParent, MaterialName, SpeedTreeMaterial, Options, WindType, GeomType, LoadedPackages, bCrossfadeLOD);
 					MaterialMap.Add(MaterialKey, StaticMesh->StaticMaterials.Num());
-					StaticMesh->StaticMaterials.Add(FStaticMaterial(Material));
+					StaticMesh->StaticMaterials.Add(FStaticMaterial(Material, FName(*MaterialName), FName(*MaterialName)));
 				}
 				uint32 MaterialIndex = MaterialMap[MaterialKey];
 
@@ -2249,10 +2250,10 @@ UObject* USpeedTreeImportFactory::FactoryCreateBinary8(UClass* InClass, UObject*
 			if (!StaticMesh->SourceModels[LODIndex].RawMeshBulkData->IsEmpty( ))
 			{
 				FRawMesh EmptyRawMesh;
-				StaticMesh->SourceModels[LODIndex].RawMeshBulkData->SaveRawMesh(EmptyRawMesh);
+				StaticMesh->SourceModels[LODIndex].SaveRawMesh(EmptyRawMesh);
 			}
 
-			StaticMesh->SourceModels[LODIndex].RawMeshBulkData->SaveRawMesh(RawMesh);
+			StaticMesh->SourceModels[LODIndex].SaveRawMesh(RawMesh);
 		}
 
 		// replace materials if they've been switched out
