Deprecate FRawMesh, replace by MeshDescription - SpeedTreeImporter

diff --git a/Engine/Plugins/Editor/SpeedTreeImporter/Source/SpeedTreeImporter/Private/SpeedTreeImportFactory.cpp b/Engine/Plugins/Editor/SpeedTreeImporter/Source/SpeedTreeImporter/Private/SpeedTreeImportFactory.cpp
--- a/Engine/Plugins/Editor/SpeedTreeImporter/Source/SpeedTreeImporter/Private/SpeedTreeImportFactory.cpp
+++ b/Engine/Plugins/Editor/SpeedTreeImporter/Source/SpeedTreeImporter/Private/SpeedTreeImportFactory.cpp
@@ -57,7 +57,9 @@
 #include "AssetRegistryModule.h"
 #include "UnrealEd/Private/GeomFitUtils.h"
 #include "SpeedTreeWind.h"
-#include "RawMesh.h"
+#include "MeshDescription.h"
+#include "MeshAttributes.h"
+#include "MeshAttributeArray.h"
 #include "ComponentReregisterContext.h"
 
 #if WITH_SPEEDTREE
@@ -1354,35 +1356,53 @@ static void MakeBodyFromCollisionObjects8(UStaticMesh* StaticMesh, SpeedTree8::D
 }
 
 
-void ProcessTriangleCorner( SpeedTree::CCore& SpeedTree, const int32 TriangleIndex, const int32 Corner, const SpeedTree::SDrawCall* DrawCall, const SpeedTree::st_uint32* Indices32, const SpeedTree::st_uint16* Indices16, FRawMesh& RawMesh, const int32 IndexOffset, const int32 NumUVs, const SpeedTree::SRenderState* RenderState )
+FVertexInstanceID ProcessTriangleCorner(
+	SpeedTree::CCore& SpeedTree,
+	const int32 TriangleIndex,
+	const int32 Corner,
+	const SpeedTree::SDrawCall* DrawCall,
+	const SpeedTree::st_uint32* Indices32,
+	const SpeedTree::st_uint16* Indices16,
+	FMeshDescription* MeshDescription,
+	const int32 IndexOffset,
+	const int32 NumUVs,
+	const SpeedTree::SRenderState* RenderState,
+	TVertexInstanceAttributeArray<FVector>& VertexInstanceNormals,
+	TVertexInstanceAttributeArray<FVector>& VertexInstanceTangents,
+	TVertexInstanceAttributeArray<float>& VertexInstanceBinormalSigns,
+	TVertexInstanceAttributeArray<FVector4>& VertexInstanceColors,
+	TVertexInstanceAttributeIndicesArray<FVector2D>& VertexInstanceUVs)
 {
+	//Speedtree use 8 UVs to store is data
+	check(VertexInstanceUVs.GetNumIndices() == 8);
+
 	SpeedTree::st_float32 Data[ 4 ];
 
 	// flip the triangle winding
 	int32 Index = TriangleIndex * 3 + Corner;
 
 	int32 VertexIndex = DrawCall->m_b32BitIndices ? Indices32[ Index ] : Indices16[ Index ];
-	RawMesh.WedgeIndices.Add( VertexIndex + IndexOffset );
+	FVertexID VertexID(IndexOffset + VertexIndex);
+	const FVertexInstanceID& VertexInstanceID = MeshDescription->CreateVertexInstance(VertexID);
 
 	// tangents
 	DrawCall->GetProperty( SpeedTree::VERTEX_PROPERTY_NORMAL, VertexIndex, Data );
 	FVector Normal( -Data[ 0 ], Data[ 1 ], Data[ 2 ] );
 	DrawCall->GetProperty( SpeedTree::VERTEX_PROPERTY_TANGENT, VertexIndex, Data );
 	FVector Tangent( -Data[ 0 ], Data[ 1 ], Data[ 2 ] );
-	RawMesh.WedgeTangentX.Add( Tangent );
-	RawMesh.WedgeTangentY.Add( Normal ^ Tangent );
-	RawMesh.WedgeTangentZ.Add( Normal );
+	VertexInstanceTangents[VertexInstanceID] = Tangent;
+	VertexInstanceNormals[VertexInstanceID] = Normal;
+	VertexInstanceBinormalSigns[VertexInstanceID] = GetBasisDeterminantSign(Tangent.GetSafeNormal(), (Normal ^ Tangent).GetSafeNormal(), Normal.GetSafeNormal());
 
 	// ao
 	DrawCall->GetProperty( SpeedTree::VERTEX_PROPERTY_AMBIENT_OCCLUSION, VertexIndex, Data );
 	uint8 AO = Data[ 0 ] * 255.0f;
-	RawMesh.WedgeColors.Add( FColor( AO, AO, AO, 255 ) );
+	VertexInstanceColors[VertexInstanceID] = FVector4(FLinearColor(FColor(AO, AO, AO, 255)));
 
 	// keep texcoords padded to align indices
-	int32 BaseTexcoordIndex = RawMesh.WedgeTexCoords[ 0 ].Num();
 	for( int32 PadIndex = 0; PadIndex < NumUVs; ++PadIndex )
 	{
-		RawMesh.WedgeTexCoords[ PadIndex ].AddZeroed( 1 );
+		VertexInstanceUVs.GetArrayForIndex(PadIndex)[VertexInstanceID] = FVector2D(0.0f, 0.0f);
 	}
 
 	// All texcoords are packed into 4 float4 vertex attributes
@@ -1402,28 +1422,28 @@ void ProcessTriangleCorner( SpeedTree::CCore& SpeedTree, const int32 TriangleInd
 
 	// diffuse
 	DrawCall->GetProperty( SpeedTree::VERTEX_PROPERTY_DIFFUSE_TEXCOORDS, VertexIndex, Data );
-	RawMesh.WedgeTexCoords[ 0 ].Top() = FVector2D( Data[ 0 ], Data[ 1 ] );
+	VertexInstanceUVs.GetArrayForIndex(0)[VertexInstanceID] = FVector2D( Data[ 0 ], Data[ 1 ] );
 
 	// lightmap
 	DrawCall->GetProperty( SpeedTree::VERTEX_PROPERTY_LIGHTMAP_TEXCOORDS, VertexIndex, Data );
-	RawMesh.WedgeTexCoords[ 1 ].Top() = FVector2D( Data[ 0 ], Data[ 1 ] );
+	VertexInstanceUVs.GetArrayForIndex(1)[VertexInstanceID] = FVector2D( Data[ 0 ], Data[ 1 ] );
 
 	// branch wind
 	DrawCall->GetProperty( SpeedTree::VERTEX_PROPERTY_WIND_BRANCH_DATA, VertexIndex, Data );
-	RawMesh.WedgeTexCoords[ 2 ].Top() = FVector2D( Data[ 0 ], Data[ 1 ] );
+	VertexInstanceUVs.GetArrayForIndex(2)[VertexInstanceID] = FVector2D( Data[ 0 ], Data[ 1 ] );
 
 	// lod
 	if( RenderState->m_bFacingLeavesPresent )
 	{
 		DrawCall->GetProperty( SpeedTree::VERTEX_PROPERTY_LEAF_CARD_LOD_SCALAR, VertexIndex, Data );
-		RawMesh.WedgeTexCoords[ 3 ].Top() = FVector2D( Data[ 0 ], 0.0f );
-		RawMesh.WedgeTexCoords[ 4 ].Top() = FVector2D( 0.0f, 0.0f );
+		VertexInstanceUVs.GetArrayForIndex(3)[VertexInstanceID] = FVector2D( Data[ 0 ], 0.0f );
+		VertexInstanceUVs.GetArrayForIndex(4)[VertexInstanceID] = FVector2D( 0.0f, 0.0f );
 	}
 	else
 	{
 		DrawCall->GetProperty( SpeedTree::VERTEX_PROPERTY_LOD_POSITION, VertexIndex, Data );
-		RawMesh.WedgeTexCoords[ 3 ].Top() = FVector2D( -Data[ 0 ], Data[ 1 ] );
-		RawMesh.WedgeTexCoords[ 4 ].Top() = FVector2D( Data[ 2 ], 0.0f );
+		VertexInstanceUVs.GetArrayForIndex(3)[VertexInstanceID] = FVector2D( -Data[ 0 ], Data[ 1 ] );
+		VertexInstanceUVs.GetArrayForIndex(4)[VertexInstanceID] = FVector2D( Data[ 2 ], 0.0f );
 	}
 
 	// other
@@ -1431,19 +1451,19 @@ void ProcessTriangleCorner( SpeedTree::CCore& SpeedTree, const int32 TriangleInd
 	{
 		// detail
 		DrawCall->GetProperty( SpeedTree::VERTEX_PROPERTY_DETAIL_TEXCOORDS, VertexIndex, Data );
-		RawMesh.WedgeTexCoords[ 5 ].Top() = FVector2D( Data[ 0 ], Data[ 1 ] );
+		VertexInstanceUVs.GetArrayForIndex(5)[VertexInstanceID] = FVector2D( Data[ 0 ], Data[ 1 ] );
 
 		// branch seam
 		DrawCall->GetProperty( SpeedTree::VERTEX_PROPERTY_BRANCH_SEAM_DIFFUSE, VertexIndex, Data );
-		RawMesh.WedgeTexCoords[ 6 ].Top() = FVector2D( Data[ 0 ], Data[ 1 ] );
-		RawMesh.WedgeTexCoords[ 4 ].Top().Y = Data[ 2 ];
+		VertexInstanceUVs.GetArrayForIndex(6)[VertexInstanceID] = FVector2D( Data[ 0 ], Data[ 1 ] );
+		VertexInstanceUVs.GetArrayForIndex(4)[VertexInstanceID].Y = Data[ 2 ];
 	}
 	else if( RenderState->m_bFrondsPresent )
 	{
 		// frond wind
 		DrawCall->GetProperty( SpeedTree::VERTEX_PROPERTY_WIND_EXTRA_DATA, VertexIndex, Data );
-		RawMesh.WedgeTexCoords[ 5 ].Top() = FVector2D( Data[ 0 ], Data[ 1 ] );
-		RawMesh.WedgeTexCoords[ 6 ].Top() = FVector2D( Data[ 2 ], 0.0f );
+		VertexInstanceUVs.GetArrayForIndex(5)[VertexInstanceID] = FVector2D( Data[ 0 ], Data[ 1 ] );
+		VertexInstanceUVs.GetArrayForIndex(6)[VertexInstanceID] = FVector2D( Data[ 2 ], 0.0f );
 	}
 	else if( RenderState->m_bLeavesPresent || RenderState->m_bFacingLeavesPresent )
 	{
@@ -1456,16 +1476,17 @@ void ProcessTriangleCorner( SpeedTree::CCore& SpeedTree, const int32 TriangleInd
 		{
 			DrawCall->GetProperty( SpeedTree::VERTEX_PROPERTY_LEAF_ANCHOR_POINT, VertexIndex, Data );
 		}
-		RawMesh.WedgeTexCoords[ 4 ].Top().Y = -Data[ 0 ];
-		RawMesh.WedgeTexCoords[ 5 ].Top() = FVector2D( Data[ 1 ], Data[ 2 ] );
+		VertexInstanceUVs.GetArrayForIndex(4)[VertexInstanceID].Y = -Data[ 0 ];
+		VertexInstanceUVs.GetArrayForIndex(5)[VertexInstanceID] = FVector2D( Data[ 1 ], Data[ 2 ] );
 
 		// leaf wind
 		DrawCall->GetProperty( SpeedTree::VERTEX_PROPERTY_WIND_EXTRA_DATA, VertexIndex, Data );
-		RawMesh.WedgeTexCoords[ 6 ].Top() = FVector2D( Data[ 0 ], Data[ 1 ] );
-		RawMesh.WedgeTexCoords[ 7 ].Top().X = Data[ 2 ];
+		VertexInstanceUVs.GetArrayForIndex(6)[VertexInstanceID] = FVector2D( Data[ 0 ], Data[ 1 ] );
+		VertexInstanceUVs.GetArrayForIndex(7)[VertexInstanceID].X = Data[ 2 ];
 		DrawCall->GetProperty( SpeedTree::VERTEX_PROPERTY_WIND_FLAGS, VertexIndex, Data );
-		RawMesh.WedgeTexCoords[ 7 ].Top().Y = Data[ 0 ];
+		VertexInstanceUVs.GetArrayForIndex(7)[VertexInstanceID].Y = Data[ 0 ];
 	}
+	return VertexInstanceID;
 }
 
 
@@ -1632,7 +1653,27 @@ UObject* USpeedTreeImportFactory::FactoryCreateBinary7(UClass* InClass, UObject*
 					for (int32 LODIndex = 0; LODIndex < SpeedTreeGeometry->m_nNumLods; ++LODIndex)
 					{
 						const SpeedTree::SLod* TreeLOD = &SpeedTreeGeometry->m_pLods[LODIndex];
-						FRawMesh RawMesh;
+						FStaticMeshSourceModel& LODModel = StaticMesh->AddSourceModel();
+						FMeshDescription* MeshDescription = StaticMesh->CreateOriginalMeshDescription(LODIndex);
+						StaticMesh->RegisterMeshAttributes(*MeshDescription);
+						
+						TVertexAttributeArray<FVector>& VertexPositions = MeshDescription->VertexAttributes().GetAttributes<FVector>(MeshAttribute::Vertex::Position);
+						TEdgeAttributeArray<bool>& EdgeHardnesses = MeshDescription->EdgeAttributes().GetAttributes<bool>(MeshAttribute::Edge::IsHard);
+						TEdgeAttributeArray<float>& EdgeCreaseSharpnesses = MeshDescription->EdgeAttributes().GetAttributes<float>(MeshAttribute::Edge::CreaseSharpness);
+						TPolygonGroupAttributeArray<FName>& PolygonGroupImportedMaterialSlotNames = MeshDescription->PolygonGroupAttributes().GetAttributes<FName>(MeshAttribute::PolygonGroup::ImportedMaterialSlotName);
+						TVertexInstanceAttributeArray<FVector>& VertexInstanceNormals = MeshDescription->VertexInstanceAttributes().GetAttributes<FVector>(MeshAttribute::VertexInstance::Normal);
+						TVertexInstanceAttributeArray<FVector>& VertexInstanceTangents = MeshDescription->VertexInstanceAttributes().GetAttributes<FVector>(MeshAttribute::VertexInstance::Tangent);
+						TVertexInstanceAttributeArray<float>& VertexInstanceBinormalSigns = MeshDescription->VertexInstanceAttributes().GetAttributes<float>(MeshAttribute::VertexInstance::BinormalSign);
+						TVertexInstanceAttributeArray<FVector4>& VertexInstanceColors = MeshDescription->VertexInstanceAttributes().GetAttributes<FVector4>(MeshAttribute::VertexInstance::Color);
+						TVertexInstanceAttributeIndicesArray<FVector2D>& VertexInstanceUVs = MeshDescription->VertexInstanceAttributes().GetAttributesSet<FVector2D>(MeshAttribute::VertexInstance::TextureCoordinate);
+
+						TMap<int32, FPolygonGroupID> MaterialToPolygonGroup;
+						for (int32 MatIndex = 0; MatIndex < StaticMesh->StaticMaterials.Num(); ++MatIndex)
+						{
+							const FPolygonGroupID& PolygonGroupID = MeshDescription->CreatePolygonGroup();
+							PolygonGroupImportedMaterialSlotNames[PolygonGroupID] = StaticMesh->StaticMaterials[MatIndex].ImportedMaterialSlotName;
+							MaterialToPolygonGroup.Add(MatIndex, PolygonGroupID);
+						}
 
 						// compute the number of texcoords we need so we can pad when necessary
 						int32 NumUVs = 7; // static meshes have fewer, but they are so rare, we shouldn't complicate things for them
@@ -1645,6 +1686,8 @@ UObject* USpeedTreeImportFactory::FactoryCreateBinary7(UClass* InClass, UObject*
 								NumUVs = FMath::Max(8, NumUVs);
 							}
 						}
+						//Speedtree use UVs to store is data
+						VertexInstanceUVs.SetNumIndices(NumUVs);
 
 						for (int32 DrawCallIndex = 0; DrawCallIndex < TreeLOD->m_nNumDrawCalls; ++DrawCallIndex)
 						{
@@ -1704,15 +1747,20 @@ UObject* USpeedTreeImportFactory::FactoryCreateBinary7(UClass* InClass, UObject*
 								UMaterialInterface* Material = CreateSpeedTreeMaterial7(InParent, MaterialName, RenderState, Options, WindType, SpeedTreeGeometry->m_sVertBBs.m_nNumBillboards, LoadedPackages);
 								
 								RenderStateIndexToStaticMeshIndex.Add(DrawCall->m_nRenderStateIndex, StaticMesh->StaticMaterials.Num());
-								MaterialIndex = StaticMesh->StaticMaterials.Num();
-								StaticMesh->StaticMaterials.Add(FStaticMaterial(Material, FName(*MaterialName), FName(*MaterialName)));
+								MaterialIndex = StaticMesh->StaticMaterials.Add(FStaticMaterial(Material, FName(*MaterialName), FName(*MaterialName)));
+								const FPolygonGroupID& PolygonGroupID = MeshDescription->CreatePolygonGroup();
+								PolygonGroupImportedMaterialSlotNames[PolygonGroupID] = StaticMesh->StaticMaterials[MaterialIndex].ImportedMaterialSlotName;
+								MaterialToPolygonGroup.Add(MaterialIndex, PolygonGroupID);
 							}
 							else
 							{
 								MaterialIndex = *OldMaterial;
 							}
 
-							int32 IndexOffset = RawMesh.VertexPositions.Num();
+
+							FPolygonGroupID CurrentPolygonGroupID = MaterialToPolygonGroup[MaterialIndex];
+
+							int32 IndexOffset = VertexPositions.Num();
 
 							for (int32 VertexIndex = 0; VertexIndex < DrawCall->m_nNumVertices; ++VertexIndex)
 							{
@@ -1727,8 +1775,8 @@ UObject* USpeedTreeImportFactory::FactoryCreateBinary7(UClass* InClass, UObject*
 									Data[1] += Data2[1];
 									Data[2] += Data2[2];                                    
 								}
-
-								RawMesh.VertexPositions.Add(FVector(-Data[0], Data[1], Data[2]));
+								FVertexID VertexID = MeshDescription->CreateVertex();
+								VertexPositions[VertexID] = FVector(-Data[0], Data[1], Data[2]);
 							}
 
 							const SpeedTree::st_byte* pIndexData = &*DrawCall->m_pIndexData;
@@ -1736,54 +1784,54 @@ UObject* USpeedTreeImportFactory::FactoryCreateBinary7(UClass* InClass, UObject*
 							const SpeedTree::st_uint16* Indices16 = (SpeedTree::st_uint16*)pIndexData;
 
 							int32 TriangleCount = DrawCall->m_nNumIndices / 3;
-							int32 ExistingTris = RawMesh.WedgeIndices.Num() / 3;
 
+							//Create all vertex instance
 							for (int32 TriangleIndex = 0; TriangleIndex < TriangleCount; ++TriangleIndex)
 							{
-								RawMesh.FaceMaterialIndices.Add(MaterialIndex);
-								RawMesh.FaceSmoothingMasks.Add(0);
-
-#if 0
-								float LeafLODScalar = 1.0f;
-								if (RenderState->m_bLeavesPresent)
+								FVertexInstanceID CornerVertexInstanceIDs[3];
+								FVertexID CornerVertexIDs[3];
+								for (int32 Corner = 0; Corner < 3; ++Corner)
 								{
-									// reconstruct a single leaf scalar so we don't need the lod position
-									// This is not 100% like the Modeler for frond-based leaves, but it reduces the amount of VB data needed
-									int32 Index = TriangleIndex * 3;
-									int32 VertexIndex1 = DrawCall->m_b32BitIndices ? Indices32[Index + 0] : Indices16[Index + 0];
-									int32 VertexIndex2 = DrawCall->m_b32BitIndices ? Indices32[Index + 1] : Indices16[Index + 1];
-									int32 VertexIndex3 = DrawCall->m_b32BitIndices ? Indices32[Index + 2] : Indices16[Index + 2];
-
-									SpeedTree::st_float32 Data[4];
-									DrawCall->GetProperty(SpeedTree::VERTEX_PROPERTY_POSITION, VertexIndex1, Data);
-									FVector Corner1(Data[0], Data[1], Data[2]);
-									DrawCall->GetProperty(SpeedTree::VERTEX_PROPERTY_POSITION, VertexIndex2, Data);
-									FVector Corner2(Data[0], Data[1], Data[2]);
-									DrawCall->GetProperty(SpeedTree::VERTEX_PROPERTY_POSITION, VertexIndex3, Data);
-									FVector Corner3(Data[0], Data[1], Data[2]);
-
-									DrawCall->GetProperty(SpeedTree::VERTEX_PROPERTY_LOD_POSITION, VertexIndex1, Data);
-									FVector LODCorner1(Data[0], Data[1], Data[2]);
-									DrawCall->GetProperty(SpeedTree::VERTEX_PROPERTY_LOD_POSITION, VertexIndex2, Data);
-									FVector LODCorner2(Data[0], Data[1], Data[2]);
-									DrawCall->GetProperty(SpeedTree::VERTEX_PROPERTY_LOD_POSITION, VertexIndex3, Data);
-									FVector LODCorner3(Data[0], Data[1], Data[2]);
-
-									float OrigArea = ((Corner2 - Corner1) ^ (Corner3 - Corner1)).Size();
-									float LODArea = ((LODCorner2 - LODCorner1) ^ (LODCorner3 - LODCorner1)).Size();
-
-									LeafLODScalar = FMath::Sqrt(LODArea / OrigArea);
+									//Create the vertex instances
+									CornerVertexInstanceIDs[Corner] = ProcessTriangleCorner( SpeedTree, TriangleIndex, Corner, DrawCall, Indices32, Indices16, MeshDescription, IndexOffset, NumUVs, RenderState,
+										VertexInstanceNormals, VertexInstanceTangents, VertexInstanceBinormalSigns, VertexInstanceColors, VertexInstanceUVs);
+
+									CornerVertexIDs[Corner] = MeshDescription->GetVertexInstanceVertex(CornerVertexInstanceIDs[Corner]);
 								}
-#endif
 
+								TArray<FMeshDescription::FContourPoint> Contours;
 								for (int32 Corner = 0; Corner < 3; ++Corner)
 								{
-									ProcessTriangleCorner( SpeedTree, TriangleIndex, Corner, DrawCall, Indices32, Indices16, RawMesh, IndexOffset, NumUVs, RenderState );
+									int32 ContourPointIndex = Contours.AddDefaulted();
+									FMeshDescription::FContourPoint& ContourPoint = Contours[ContourPointIndex];
+									//Find the matching edge ID
+									uint32 CornerIndices[2];
+									CornerIndices[0] = (Corner + 0) % 3;
+									CornerIndices[1] = (Corner + 1) % 3;
+
+									FVertexID EdgeVertexIDs[2];
+									EdgeVertexIDs[0] = CornerVertexIDs[CornerIndices[0]];
+									EdgeVertexIDs[1] = CornerVertexIDs[CornerIndices[1]];
+
+									FEdgeID MatchEdgeId = MeshDescription->GetVertexPairEdge(EdgeVertexIDs[0], EdgeVertexIDs[1]);
+									if (MatchEdgeId == FEdgeID::Invalid)
+									{
+										MatchEdgeId = MeshDescription->CreateEdge(EdgeVertexIDs[0], EdgeVertexIDs[1]);
+										//All edge are smooth
+										EdgeHardnesses[MatchEdgeId] = false;
+										EdgeCreaseSharpnesses[MatchEdgeId] = 0.0f;
+									}
+									ContourPoint.EdgeID = MatchEdgeId;
+									ContourPoint.VertexInstanceID = CornerVertexInstanceIDs[CornerIndices[0]];
 								}
-							}							
+								// Insert a polygon into the mesh
+								const FPolygonID NewPolygonID = MeshDescription->CreatePolygon(CurrentPolygonGroupID, Contours);
+								//Triangulate the polygon
+								FMeshPolygon& Polygon = MeshDescription->GetPolygon(NewPolygonID);
+								MeshDescription->ComputePolygonTriangulation(NewPolygonID, Polygon.Triangles);
+							}
 						}
 
-						FStaticMeshSourceModel& LODModel = StaticMesh->AddSourceModel();
 						LODModel.BuildSettings.bRecomputeNormals = false;
 						LODModel.BuildSettings.bRecomputeTangents = false;
 						LODModel.BuildSettings.bRemoveDegenerates = true;
@@ -1791,7 +1839,7 @@ UObject* USpeedTreeImportFactory::FactoryCreateBinary7(UClass* InClass, UObject*
 						LODModel.BuildSettings.bUseFullPrecisionUVs = false;
 						LODModel.BuildSettings.bGenerateLightmapUVs = false;
 						LODModel.ScreenSize.Default = 0.1f / FMath::Pow(2.0f, StaticMesh->SourceModels.Num() - 1);
-						LODModel.SaveRawMesh(RawMesh);
+						StaticMesh->CommitOriginalMeshDescription(LODIndex);
 
 						for (int32 MaterialIndex = 0; MaterialIndex < StaticMesh->StaticMaterials.Num(); ++MaterialIndex)
 						{
@@ -1806,12 +1854,36 @@ UObject* USpeedTreeImportFactory::FactoryCreateBinary7(UClass* InClass, UObject*
 				// make billboard LOD
 				if (Options->SpeedTreeImportData->ImportGeometryType != EImportGeometryType::IGT_3D && SpeedTreeGeometry->m_sVertBBs.m_nNumBillboards > 0)
 				{
+					FStaticMeshSourceModel& LODModel = StaticMesh->AddSourceModel();
+					const int32 LODIndex = StaticMesh->SourceModels.Num() - 1;
+					FMeshDescription* MeshDescription = StaticMesh->CreateOriginalMeshDescription(LODIndex);
+					StaticMesh->RegisterMeshAttributes(*MeshDescription);
+
+					TVertexAttributeArray<FVector>& VertexPositions = MeshDescription->VertexAttributes().GetAttributes<FVector>(MeshAttribute::Vertex::Position);
+					TEdgeAttributeArray<bool>& EdgeHardnesses = MeshDescription->EdgeAttributes().GetAttributes<bool>(MeshAttribute::Edge::IsHard);
+					TEdgeAttributeArray<float>& EdgeCreaseSharpnesses = MeshDescription->EdgeAttributes().GetAttributes<float>(MeshAttribute::Edge::CreaseSharpness);
+					TPolygonGroupAttributeArray<FName>& PolygonGroupImportedMaterialSlotNames = MeshDescription->PolygonGroupAttributes().GetAttributes<FName>(MeshAttribute::PolygonGroup::ImportedMaterialSlotName);
+					TVertexInstanceAttributeArray<FVector>& VertexInstanceNormals = MeshDescription->VertexInstanceAttributes().GetAttributes<FVector>(MeshAttribute::VertexInstance::Normal);
+					TVertexInstanceAttributeArray<FVector>& VertexInstanceTangents = MeshDescription->VertexInstanceAttributes().GetAttributes<FVector>(MeshAttribute::VertexInstance::Tangent);
+					TVertexInstanceAttributeArray<float>& VertexInstanceBinormalSigns = MeshDescription->VertexInstanceAttributes().GetAttributes<float>(MeshAttribute::VertexInstance::BinormalSign);
+					TVertexInstanceAttributeArray<FVector4>& VertexInstanceColors = MeshDescription->VertexInstanceAttributes().GetAttributes<FVector4>(MeshAttribute::VertexInstance::Color);
+					TVertexInstanceAttributeIndicesArray<FVector2D>& VertexInstanceUVs = MeshDescription->VertexInstanceAttributes().GetAttributesSet<FVector2D>(MeshAttribute::VertexInstance::TextureCoordinate);
+
+					//Speedtree use UVs to store is data
+					VertexInstanceUVs.SetNumIndices(2);
+
+					for (int32 MatIndex = 0; MatIndex < StaticMesh->StaticMaterials.Num(); ++MatIndex)
+					{
+						const FPolygonGroupID& PolygonGroupID = MeshDescription->CreatePolygonGroup();
+						PolygonGroupImportedMaterialSlotNames[PolygonGroupID] = StaticMesh->StaticMaterials[MatIndex].ImportedMaterialSlotName;
+					}
+
 					FString MaterialName = MeshName + "_Billboard";
 					UMaterialInterface* Material = CreateSpeedTreeMaterial7(InParent, MaterialName, &SpeedTreeGeometry->m_aBillboardRenderStates[SpeedTree::RENDER_PASS_MAIN], Options, WindType, SpeedTreeGeometry->m_sVertBBs.m_nNumBillboards, LoadedPackages);
-					int32 MaterialIndex = StaticMesh->StaticMaterials.Num();
-					StaticMesh->StaticMaterials.Add(FStaticMaterial(Material, FName(*MaterialName), FName(*MaterialName)));
+					int32 MaterialIndex = StaticMesh->StaticMaterials.Add(FStaticMaterial(Material, FName(*MaterialName), FName(*MaterialName)));
 
-					FRawMesh RawMesh;
+					const FPolygonGroupID& CurrentPolygonGroupID = MeshDescription->CreatePolygonGroup();
+					PolygonGroupImportedMaterialSlotNames[CurrentPolygonGroupID] = StaticMesh->StaticMaterials[MaterialIndex].ImportedMaterialSlotName;
 					
 					// fill out triangles
 					float BillboardWidth = SpeedTreeGeometry->m_sVertBBs.m_fWidth;
@@ -1855,50 +1927,84 @@ UObject* USpeedTreeImportFactory::FactoryCreateBinary7(UClass* InClass, UObject*
 						const float* TexCoords = &SpeedTreeGeometry->m_sVertBBs.m_pTexCoords[BillboardIndex * 4];
 						bool bRotated = (SpeedTreeGeometry->m_sVertBBs.m_pRotated[BillboardIndex] == 1);
 
-						int32 IndexOffset = RawMesh.VertexPositions.Num();
+						int32 IndexOffset = VertexPositions.Num();
 					
 						// position
 						for (int32 VertexIndex = 0; VertexIndex < NumVertices; ++VertexIndex)
 						{
 							const SpeedTree::st_float32* Vertex = &Vertices[VertexIndex * 2];
 							FVector Position = BillboardRotate.TransformVector(FVector(Vertex[0] * BillboardWidth - BillboardWidth * 0.5f, 0.0f, Vertex[1] * BillboardHeight + BillboardBottom));
-							RawMesh.VertexPositions.Add(Position);
+							
+							FVertexID VertexID = MeshDescription->CreateVertex();
+							VertexPositions[VertexID] = FVector(Position);
 						}
 
 						// other data
 						int32 NumTriangles = NumIndices / 3;
 						for (int32 TriangleIndex = 0; TriangleIndex < NumTriangles; ++TriangleIndex)
 						{
-							RawMesh.FaceMaterialIndices.Add(MaterialIndex);
-							RawMesh.FaceSmoothingMasks.Add(0);
-
+							FVertexInstanceID CornerVertexInstanceIDs[3];
+							FVertexID CornerVertexIDs[3];
 							for (int32 Corner = 0; Corner < 3; ++Corner)
 							{
 								int32 Index = Indices[TriangleIndex * 3 + Corner];
 								const SpeedTree::st_float32* Vertex = &Vertices[Index * 2];
 
-								RawMesh.WedgeIndices.Add(Index + IndexOffset);
-
-								RawMesh.WedgeTangentX.Add(TangentX);
-								RawMesh.WedgeTangentY.Add(TangentY);
-								RawMesh.WedgeTangentZ.Add(TangentZ);
+								FVertexID VertexID(IndexOffset + Index);
+								const FVertexInstanceID& VertexInstanceID = MeshDescription->CreateVertexInstance(VertexID);
 
+								VertexInstanceTangents[VertexInstanceID] = TangentX;
+								VertexInstanceNormals[VertexInstanceID] = TangentZ;
+								VertexInstanceBinormalSigns[VertexInstanceID] = GetBasisDeterminantSign(TangentX.GetSafeNormal(), TangentY.GetSafeNormal(), TangentZ.GetSafeNormal());
 								if (bRotated)
 								{
-									RawMesh.WedgeTexCoords[0].Add(FVector2D(TexCoords[0] + Vertex[1] * TexCoords[2], TexCoords[1] + Vertex[0] * TexCoords[3]));
+									VertexInstanceUVs.GetArrayForIndex(0)[VertexInstanceID] = FVector2D(TexCoords[0] + Vertex[1] * TexCoords[2], TexCoords[1] + Vertex[0] * TexCoords[3]);
 								}
 								else
 								{
-									RawMesh.WedgeTexCoords[0].Add(FVector2D(TexCoords[0] + Vertex[0] * TexCoords[2], TexCoords[1] + Vertex[1] * TexCoords[3]));
+									VertexInstanceUVs.GetArrayForIndex(0)[VertexInstanceID] = FVector2D(TexCoords[0] + Vertex[0] * TexCoords[2], TexCoords[1] + Vertex[1] * TexCoords[3]);
 								}
 
 								// lightmap coord
-								RawMesh.WedgeTexCoords[1].Add(RawMesh.WedgeTexCoords[0].Top());
+								VertexInstanceUVs.GetArrayForIndex(1)[VertexInstanceID] = VertexInstanceUVs.GetArrayForIndex(0)[VertexInstanceID];
+
+								CornerVertexInstanceIDs[Corner] = VertexInstanceID;
+								CornerVertexIDs[Corner] = VertexID;
 							}
+
+							TArray<FMeshDescription::FContourPoint> Contours;
+							for (int32 Corner = 0; Corner < 3; ++Corner)
+							{
+								int32 ContourPointIndex = Contours.AddDefaulted();
+								FMeshDescription::FContourPoint& ContourPoint = Contours[ContourPointIndex];
+								//Find the matching edge ID
+								uint32 CornerIndices[2];
+								CornerIndices[0] = (Corner + 0) % 3;
+								CornerIndices[1] = (Corner + 1) % 3;
+
+								FVertexID EdgeVertexIDs[2];
+								EdgeVertexIDs[0] = CornerVertexIDs[CornerIndices[0]];
+								EdgeVertexIDs[1] = CornerVertexIDs[CornerIndices[1]];
+
+								FEdgeID MatchEdgeId = MeshDescription->GetVertexPairEdge(EdgeVertexIDs[0], EdgeVertexIDs[1]);
+								if (MatchEdgeId == FEdgeID::Invalid)
+								{
+									MatchEdgeId = MeshDescription->CreateEdge(EdgeVertexIDs[0], EdgeVertexIDs[1]);
+									//All edge are smooth
+									EdgeHardnesses[MatchEdgeId] = false;
+									EdgeCreaseSharpnesses[MatchEdgeId] = 0.0f;
+								}
+								ContourPoint.EdgeID = MatchEdgeId;
+								ContourPoint.VertexInstanceID = CornerVertexInstanceIDs[CornerIndices[0]];
+							}
+							// Insert a polygon into the mesh
+							const FPolygonID NewPolygonID = MeshDescription->CreatePolygon(CurrentPolygonGroupID, Contours);
+							//Triangulate the polygon
+							FMeshPolygon& Polygon = MeshDescription->GetPolygon(NewPolygonID);
+							MeshDescription->ComputePolygonTriangulation(NewPolygonID, Polygon.Triangles);
 						}
 					}
 
-					FStaticMeshSourceModel& LODModel = StaticMesh->AddSourceModel();
 					LODModel.BuildSettings.bRecomputeNormals = false;
 					LODModel.BuildSettings.bRecomputeTangents = false;
 					LODModel.BuildSettings.bRemoveDegenerates = true;
@@ -1906,9 +2012,8 @@ UObject* USpeedTreeImportFactory::FactoryCreateBinary7(UClass* InClass, UObject*
 					LODModel.BuildSettings.bUseFullPrecisionUVs = false;
 					LODModel.BuildSettings.bGenerateLightmapUVs = false;
 					LODModel.ScreenSize.Default = 0.1f / FMath::Pow(2.0f, StaticMesh->SourceModels.Num() - 1);
-					LODModel.SaveRawMesh(RawMesh);
+					StaticMesh->CommitOriginalMeshDescription(LODIndex);
 					// Add mesh section info entry for billboard LOD (only one section/material index)
-					const int32 LODIndex = StaticMesh->SourceModels.Num() - 1;
 					FMeshSectionInfo Info = StaticMesh->SectionInfoMap.Get(LODIndex, 0);
 					Info.MaterialIndex = MaterialIndex;
 					StaticMesh->SectionInfoMap.Set(LODIndex, 0, Info);
@@ -2094,7 +2199,7 @@ UObject* USpeedTreeImportFactory::FactoryCreateBinary8(UClass* InClass, UObject*
 		};
 
 		// materials
-		TMap<FIntPoint, uint32> MaterialMap;
+		TMap<FIntPoint, int32> MaterialMap;
 
 		// gui info
 		StaticMesh->bAutoComputeLODScreenSize = false;
@@ -2106,13 +2211,35 @@ UObject* USpeedTreeImportFactory::FactoryCreateBinary8(UClass* InClass, UObject*
 		for (uint32 LODIndex = 0; LODIndex < SpeedTree.Lods().Count(); ++LODIndex)
 		{
 			SpeedTree8::CLod LOD = SpeedTree.Lods()[LODIndex];
-			FRawMesh RawMesh;
+			
+			FMeshDescription* MeshDescription = StaticMesh->CreateOriginalMeshDescription(LODIndex);
+			StaticMesh->RegisterMeshAttributes(*MeshDescription);
+
+			TVertexAttributeArray<FVector>& VertexPositions = MeshDescription->VertexAttributes().GetAttributes<FVector>(MeshAttribute::Vertex::Position);
+			TEdgeAttributeArray<bool>& EdgeHardnesses = MeshDescription->EdgeAttributes().GetAttributes<bool>(MeshAttribute::Edge::IsHard);
+			TEdgeAttributeArray<float>& EdgeCreaseSharpnesses = MeshDescription->EdgeAttributes().GetAttributes<float>(MeshAttribute::Edge::CreaseSharpness);
+			TPolygonGroupAttributeArray<FName>& PolygonGroupImportedMaterialSlotNames = MeshDescription->PolygonGroupAttributes().GetAttributes<FName>(MeshAttribute::PolygonGroup::ImportedMaterialSlotName);
+			TVertexInstanceAttributeArray<FVector>& VertexInstanceNormals = MeshDescription->VertexInstanceAttributes().GetAttributes<FVector>(MeshAttribute::VertexInstance::Normal);
+			TVertexInstanceAttributeArray<FVector>& VertexInstanceTangents = MeshDescription->VertexInstanceAttributes().GetAttributes<FVector>(MeshAttribute::VertexInstance::Tangent);
+			TVertexInstanceAttributeArray<float>& VertexInstanceBinormalSigns = MeshDescription->VertexInstanceAttributes().GetAttributes<float>(MeshAttribute::VertexInstance::BinormalSign);
+			TVertexInstanceAttributeArray<FVector4>& VertexInstanceColors = MeshDescription->VertexInstanceAttributes().GetAttributes<FVector4>(MeshAttribute::VertexInstance::Color);
+			TVertexInstanceAttributeIndicesArray<FVector2D>& VertexInstanceUVs = MeshDescription->VertexInstanceAttributes().GetAttributesSet<FVector2D>(MeshAttribute::VertexInstance::TextureCoordinate);
+
+			//Speedtree use 8 UVs to store is data
+			VertexInstanceUVs.SetNumIndices(8);
+
+			for (int32 MatIndex = 0; MatIndex < StaticMesh->StaticMaterials.Num(); ++MatIndex)
+			{
+				const FPolygonGroupID& PolygonGroupID = MeshDescription->CreatePolygonGroup();
+				PolygonGroupImportedMaterialSlotNames[PolygonGroupID] = StaticMesh->StaticMaterials[MatIndex].ImportedMaterialSlotName;
+			}
 
 			for (uint32 VertexIndex = 0; VertexIndex < LOD.Vertices().Count(); ++VertexIndex)
 			{
 				const SpeedTree8::SVertex& Vertex = LOD.Vertices()[VertexIndex]; //-V758
 				FVector vPosition = FVector(Vertex.m_vAnchor.x, Vertex.m_vAnchor.y, Vertex.m_vAnchor.z) + FVector(Vertex.m_vOffset.x, Vertex.m_vOffset.y, Vertex.m_vOffset.z);
-				RawMesh.VertexPositions.Add(vPosition);
+				FVertexID VertexID = MeshDescription->CreateVertex();
+				VertexPositions[VertexID] = FVector(vPosition);
 			}
 
 			for (uint32 DrawCallIndex = 0; DrawCallIndex < LOD.DrawCalls().Count(); ++DrawCallIndex)
@@ -2141,9 +2268,12 @@ UObject* USpeedTreeImportFactory::FactoryCreateBinary8(UClass* InClass, UObject*
 					MaterialName.InsertAt(MaterialName.Len() - 4, GeomString);
 					UMaterialInterface* Material = CreateSpeedTreeMaterial8(InParent, MaterialName, SpeedTreeMaterial, Options, WindType, GeomType, LoadedPackages, bCrossfadeLOD);
 					MaterialMap.Add(MaterialKey, StaticMesh->StaticMaterials.Num());
-					StaticMesh->StaticMaterials.Add(FStaticMaterial(Material, FName(*MaterialName), FName(*MaterialName)));
+					int32 AddedMaterialIndex = StaticMesh->StaticMaterials.Add(FStaticMaterial(Material, FName(*MaterialName), FName(*MaterialName)));
+					const FPolygonGroupID& PolygonGroupID = MeshDescription->CreatePolygonGroup();
+					PolygonGroupImportedMaterialSlotNames[PolygonGroupID] = StaticMesh->StaticMaterials[AddedMaterialIndex].ImportedMaterialSlotName;
 				}
-				uint32 MaterialIndex = MaterialMap[MaterialKey];
+				const int32 MaterialIndex = MaterialMap[MaterialKey];
+				const FPolygonGroupID CurrentPolygonGroupID(MaterialIndex);
 
 				FMeshSectionInfo Info = StaticMesh->SectionInfoMap.Get(LODIndex, DrawCallIndex);
 				Info.MaterialIndex = MaterialIndex;
@@ -2152,26 +2282,30 @@ UObject* USpeedTreeImportFactory::FactoryCreateBinary8(UClass* InClass, UObject*
 				int32 TriangleCount = DrawCall.m_uiIndexCount / 3;
 				for (int32 TriangleIndex = 0; TriangleIndex < TriangleCount; ++TriangleIndex)
 				{
-					RawMesh.FaceMaterialIndices.Add(MaterialIndex);
-					RawMesh.FaceSmoothingMasks.Add(0);
-
+					FVertexInstanceID CornerVertexInstanceIDs[3];
+					FVertexID CornerVertexIDs[3];
 					for (int32 Corner = 0; Corner < 3; ++Corner)
 					{
 						int32 VertexIndex = LOD.Indices()[DrawCall.m_uiIndexStart + TriangleIndex * 3 + (2 - Corner)];
-						RawMesh.WedgeIndices.Add(VertexIndex);
-
 						const SpeedTree8::SVertex& Vertex = LOD.Vertices()[VertexIndex]; //-V758
 
-						// tangents
-						RawMesh.WedgeTangentX.Add(FVector(Vertex.m_vTangent.x, Vertex.m_vTangent.y, Vertex.m_vTangent.z));
-						RawMesh.WedgeTangentY.Add(FVector(-Vertex.m_vBinormal.x, -Vertex.m_vBinormal.y, -Vertex.m_vBinormal.z));
-						RawMesh.WedgeTangentZ.Add(FVector(Vertex.m_vNormal.x, Vertex.m_vNormal.y, Vertex.m_vNormal.z));
+						FVertexID VertexID(VertexIndex);
+						const FVertexInstanceID& VertexInstanceID = MeshDescription->CreateVertexInstance(VertexID);
 
+						// tangents
+						FVector TangentX(Vertex.m_vTangent.x, Vertex.m_vTangent.y, Vertex.m_vTangent.z);
+						FVector TangentY(-Vertex.m_vBinormal.x, -Vertex.m_vBinormal.y, -Vertex.m_vBinormal.z);
+						FVector TangentZ(Vertex.m_vNormal.x, Vertex.m_vNormal.y, Vertex.m_vNormal.z);
+
+						VertexInstanceTangents[VertexInstanceID] = TangentX;
+						VertexInstanceNormals[VertexInstanceID] = TangentZ;
+						VertexInstanceBinormalSigns[VertexInstanceID] = GetBasisDeterminantSign(TangentX.GetSafeNormal(), TangentY.GetSafeNormal(), TangentZ.GetSafeNormal());
+						
 						// ao and branch blend in vertex color
 						if (DrawCall.m_eWindGeometryType != SpeedTree8::Billboard)
 						{
 							uint8 AO = Vertex.m_fAmbientOcclusion * 255.0f;
-							RawMesh.WedgeColors.Add(FColor(AO, AO, AO, Vertex.m_fBlendWeight * 255));
+							VertexInstanceColors[VertexInstanceID] = FVector4(FLinearColor(FColor(AO, AO, AO, Vertex.m_fBlendWeight * 255)));
 						}
 
 						// All texcoords are packed into 4 float4 vertex attributes
@@ -2189,71 +2323,98 @@ UObject* USpeedTreeImportFactory::FactoryCreateBinary8(UClass* InClass, UObject*
 						// 7	0					0					Leaf Wind Z, Leaf Group
 
 						// diffuse
-						RawMesh.WedgeTexCoords[0].Add(FVector2D(Vertex.m_vTexCoord.x, Vertex.m_vTexCoord.y));
+						VertexInstanceUVs.GetArrayForIndex(0)[VertexInstanceID] = FVector2D(Vertex.m_vTexCoord.x, Vertex.m_vTexCoord.y);
 
 						// lightmap
-						RawMesh.WedgeTexCoords[1].Add(FVector2D(Vertex.m_vLightmapTexCoord.x, Vertex.m_vLightmapTexCoord.y));
+						VertexInstanceUVs.GetArrayForIndex(1)[VertexInstanceID] = FVector2D(Vertex.m_vLightmapTexCoord.x, Vertex.m_vLightmapTexCoord.y);
 
 						if (DrawCall.m_eWindGeometryType == SpeedTree8::Billboard)
 						{
-							RawMesh.WedgeTexCoords[2].Add(FVector2D((Vertex.m_vNormal.z > 0.5f) ? 1.0f : 0.0f, 0.0f));
+							VertexInstanceUVs.GetArrayForIndex(2)[VertexInstanceID] = FVector2D((Vertex.m_vNormal.z > 0.5f) ? 1.0f : 0.0f, 0.0f);
 						}
 						else
 						{
 							// branch wind
-							RawMesh.WedgeTexCoords[2].Add(FVector2D(Vertex.m_vWindBranch.x, Vertex.m_vWindBranch.y));
+							VertexInstanceUVs.GetArrayForIndex(2)[VertexInstanceID] = FVector2D(Vertex.m_vWindBranch.x, Vertex.m_vWindBranch.y);
 
 							// lod
 							FVector vLodPosition = FVector(Vertex.m_vAnchor.x, Vertex.m_vAnchor.y, Vertex.m_vAnchor.z) +
-													FVector(Vertex.m_vLodOffset.x, Vertex.m_vLodOffset.y, Vertex.m_vLodOffset.z);
-							RawMesh.WedgeTexCoords[3].Add(FVector2D(vLodPosition[0], vLodPosition[1]));
-							RawMesh.WedgeTexCoords[4].Add(FVector2D(vLodPosition[2], 0.0f));
+								FVector(Vertex.m_vLodOffset.x, Vertex.m_vLodOffset.y, Vertex.m_vLodOffset.z);
+							VertexInstanceUVs.GetArrayForIndex(3)[VertexInstanceID] = FVector2D(vLodPosition[0], vLodPosition[1]);
+							VertexInstanceUVs.GetArrayForIndex(4)[VertexInstanceID] = FVector2D(vLodPosition[2], 0.0f);
 
 							// other
 							if (DrawCall.m_eWindGeometryType == SpeedTree8::Branch)
 							{
 								// detail (not used in v8)
-								RawMesh.WedgeTexCoords[5].Add(FVector2D(0.0f, 0.0f));
+								VertexInstanceUVs.GetArrayForIndex(5)[VertexInstanceID] = FVector2D(0.0f, 0.0f);
 
 								// branch seam
-								RawMesh.WedgeTexCoords[6].Add(FVector2D(0.0f, 0.0f));
-								RawMesh.WedgeTexCoords[4].Top().Y = Vertex.m_fBlendWeight;
+								VertexInstanceUVs.GetArrayForIndex(6)[VertexInstanceID] = FVector2D(0.0f, 0.0f);
+								VertexInstanceUVs.GetArrayForIndex(4)[VertexInstanceID].Y = Vertex.m_fBlendWeight;
 
 								// keep alignment
-								RawMesh.WedgeTexCoords[7].Add(FVector2D(0.0f, 0.0f));
+								VertexInstanceUVs.GetArrayForIndex(7)[VertexInstanceID] = FVector2D(0.0f, 0.0f);
 							}
-							 else if (DrawCall.m_eWindGeometryType == SpeedTree8::Frond)
+							else if (DrawCall.m_eWindGeometryType == SpeedTree8::Frond)
 							{
 								// frond wind
-								RawMesh.WedgeTexCoords[5].Add(FVector2D(Vertex.m_vWindNonBranch.x, Vertex.m_vWindNonBranch.y));
-								RawMesh.WedgeTexCoords[6].Add(FVector2D(Vertex.m_vWindNonBranch.z, 0.0f));
+								VertexInstanceUVs.GetArrayForIndex(5)[VertexInstanceID] = FVector2D(Vertex.m_vWindNonBranch.x, Vertex.m_vWindNonBranch.y);
+								VertexInstanceUVs.GetArrayForIndex(6)[VertexInstanceID] = FVector2D(Vertex.m_vWindNonBranch.z, 0.0f);
 
 								// keep alignment
-								RawMesh.WedgeTexCoords[7].Add(FVector2D(0.0f, 0.0f));
+								VertexInstanceUVs.GetArrayForIndex(7)[VertexInstanceID] = FVector2D(0.0f, 0.0f);
 							}
 							else if (DrawCall.m_eWindGeometryType == SpeedTree8::Leaf || DrawCall.m_eWindGeometryType == SpeedTree8::FacingLeaf)
 							{
 								// anchor
-								RawMesh.WedgeTexCoords[4].Top().Y = Vertex.m_vAnchor.x;
-								RawMesh.WedgeTexCoords[5].Add(FVector2D(Vertex.m_vAnchor.y, Vertex.m_vAnchor.z));
+								VertexInstanceUVs.GetArrayForIndex(4)[VertexInstanceID].Y = Vertex.m_vAnchor.x;
+								VertexInstanceUVs.GetArrayForIndex(5)[VertexInstanceID] = FVector2D(Vertex.m_vAnchor.y, Vertex.m_vAnchor.z);
 
 								// leaf wind
-								RawMesh.WedgeTexCoords[6].Add(FVector2D(Vertex.m_vWindNonBranch.x, Vertex.m_vWindNonBranch.y));
-								RawMesh.WedgeTexCoords[7].Add(FVector2D(Vertex.m_vWindNonBranch.z, (Vertex.m_bWindLeaf2Flag ? 1.0f : 0.0f)));
+								VertexInstanceUVs.GetArrayForIndex(6)[VertexInstanceID] = FVector2D(Vertex.m_vWindNonBranch.x, Vertex.m_vWindNonBranch.y);
+								VertexInstanceUVs.GetArrayForIndex(7)[VertexInstanceID] = FVector2D(Vertex.m_vWindNonBranch.z, (Vertex.m_bWindLeaf2Flag ? 1.0f : 0.0f));
 							}
 						}
+						CornerVertexInstanceIDs[Corner] = VertexInstanceID;
+						CornerVertexIDs[Corner] = VertexID;
 					}
-				}
-			}
 
-
-			if (!StaticMesh->SourceModels[LODIndex].RawMeshBulkData->IsEmpty( ))
-			{
-				FRawMesh EmptyRawMesh;
-				StaticMesh->SourceModels[LODIndex].SaveRawMesh(EmptyRawMesh);
+					TArray<FMeshDescription::FContourPoint> Contours;
+					for (int32 Corner = 0; Corner < 3; ++Corner)
+					{
+						int32 ContourPointIndex = Contours.AddDefaulted();
+						FMeshDescription::FContourPoint& ContourPoint = Contours[ContourPointIndex];
+						//Find the matching edge ID
+						uint32 CornerIndices[2];
+						CornerIndices[0] = (Corner + 0) % 3;
+						CornerIndices[1] = (Corner + 1) % 3;
+
+						FVertexID EdgeVertexIDs[2];
+						EdgeVertexIDs[0] = CornerVertexIDs[CornerIndices[0]];
+						EdgeVertexIDs[1] = CornerVertexIDs[CornerIndices[1]];
+
+						FEdgeID MatchEdgeId = MeshDescription->GetVertexPairEdge(EdgeVertexIDs[0], EdgeVertexIDs[1]);
+						if (MatchEdgeId == FEdgeID::Invalid)
+						{
+							MatchEdgeId = MeshDescription->CreateEdge(EdgeVertexIDs[0], EdgeVertexIDs[1]);
+							//All edge are smooth
+							EdgeHardnesses[MatchEdgeId] = false;
+							EdgeCreaseSharpnesses[MatchEdgeId] = 0.0f;
+						}
+						ContourPoint.EdgeID = MatchEdgeId;
+						ContourPoint.VertexInstanceID = CornerVertexInstanceIDs[CornerIndices[0]];
+					}
+					// Insert a polygon into the mesh
+					const FPolygonID NewPolygonID = MeshDescription->CreatePolygon(CurrentPolygonGroupID, Contours);
+					//Triangulate the polygon
+					FMeshPolygon& Polygon = MeshDescription->GetPolygon(NewPolygonID);
+					MeshDescription->ComputePolygonTriangulation(NewPolygonID, Polygon.Triangles);
+				}
 			}
 
-			StaticMesh->SourceModels[LODIndex].SaveRawMesh(RawMesh);
+			//Save the created mesh
+			StaticMesh->CommitOriginalMeshDescription(LODIndex);
 		}
 
 		// replace materials if they've been switched out
diff --git a/Engine/Plugins/Editor/SpeedTreeImporter/Source/SpeedTreeImporter/SpeedTreeImporter.Build.cs b/Engine/Plugins/Editor/SpeedTreeImporter/Source/SpeedTreeImporter/SpeedTreeImporter.Build.cs
--- a/Engine/Plugins/Editor/SpeedTreeImporter/Source/SpeedTreeImporter/SpeedTreeImporter.Build.cs
+++ b/Engine/Plugins/Editor/SpeedTreeImporter/Source/SpeedTreeImporter/SpeedTreeImporter.Build.cs
@@ -32,10 +32,10 @@ namespace UnrealBuildTool.Rules
 					"SlateCore",
                     "EditorStyle",
                     "InputCore",
-				    "RawMesh",
                     "UnrealEd",
-                    "MainFrame"
-				}
+                    "MainFrame",
+                    "MeshDescription",
+                }
 				);
 
 			DynamicallyLoadedModuleNames.AddRange(
