Fix category label on speedtree importer plugin

diff --git a/Engine/Plugins/Editor/SpeedTreeImporter/SpeedTreeImporter.uplugin b/Engine/Plugins/Editor/SpeedTreeImporter/SpeedTreeImporter.uplugin
--- a/Engine/Plugins/Editor/SpeedTreeImporter/SpeedTreeImporter.uplugin
+++ b/Engine/Plugins/Editor/SpeedTreeImporter/SpeedTreeImporter.uplugin
@@ -1,25 +1,24 @@
 {
-	"FileVersion" : 3,
-	"Version" : 1,
-	"VersionName" : "1.0",
-	"FriendlyName" : "SpeedTree Importer",
-	"Description" : "An importer for SpeedTree runtime (SRT) files.",
-	"Category" : "Editor",
-	"CreatedBy" : "IDV, Inc.",
-	"CreatedByURL" : "http://www.speedtree.com",
-	"DocsURL" : "",
-	"MarketplaceURL" : "",
-	"SupportURL" : "",
-	"EnabledByDefault" : true,
-	"CanContainContent" : false,
-	"IsBetaVersion" : false,
-	"Installed" : false,
-	"Modules" :
-	[
+	"FileVersion": 3,
+	"Version": 1,
+	"VersionName": "1.0",
+	"FriendlyName": "SpeedTree Importer",
+	"Description": "An importer for SpeedTree runtime (SRT) files.",
+	"Category": "Importers",
+	"CreatedBy": "IDV, Inc.",
+	"CreatedByURL": "http://www.speedtree.com",
+	"DocsURL": "",
+	"MarketplaceURL": "",
+	"SupportURL": "",
+	"EnabledByDefault": true,
+	"CanContainContent": false,
+	"IsBetaVersion": false,
+	"Installed": false,
+	"Modules": [
 		{
-			"Name" : "SpeedTreeImporter",
-			"Type" : "Editor",
-			"LoadingPhase" : "Default"
+			"Name": "SpeedTreeImporter",
+			"Type": "Editor",
+			"LoadingPhase": "Default"
 		}
 	]
 }
\ No newline at end of file
