Fixed crash bug.

diff --git a/Engine/Plugins/Editor/SpeedTreeImporter/Source/SpeedTreeImporter/Private/SpeedTreeImportFactory.cpp b/Engine/Plugins/Editor/SpeedTreeImporter/Source/SpeedTreeImporter/Private/SpeedTreeImportFactory.cpp
--- a/Engine/Plugins/Editor/SpeedTreeImporter/Source/SpeedTreeImporter/Private/SpeedTreeImportFactory.cpp
+++ b/Engine/Plugins/Editor/SpeedTreeImporter/Source/SpeedTreeImporter/Private/SpeedTreeImportFactory.cpp
@@ -1366,8 +1366,8 @@ FVertexInstanceID ProcessTriangleCorner(
 	TVertexInstanceAttributesRef<FVector4> VertexInstanceColors,
 	TVertexInstanceAttributesRef<FVector2D> VertexInstanceUVs)
 {
-	//Speedtree use 8 UVs to store is data
-	check(VertexInstanceUVs.GetNumIndices() == 8);
+	//Speedtree uses 7 or 8 UVs to store is data
+	check(VertexInstanceUVs.GetNumIndices() >= 7);
 
 	SpeedTree::st_float32 Data[ 4 ];
 
@@ -1460,6 +1460,8 @@ FVertexInstanceID ProcessTriangleCorner(
 	}
 	else if( RenderState->m_bLeavesPresent || RenderState->m_bFacingLeavesPresent )
 	{
+		check(VertexInstanceUVs.GetNumIndices() == 8);
+
 		// anchor
 		if( RenderState->m_bFacingLeavesPresent )
 		{
