Copying //UE4/Dev-Editor to Dev-Main (//UE4/Dev-Main) for 4.24

diff --git a/Engine/Plugins/Editor/GLTFImporter/Source/GLTFImporter/GLTFImporter.Build.cs b/Engine/Plugins/Editor/GLTFImporter/Source/GLTFImporter/GLTFImporter.Build.cs
--- a/Engine/Plugins/Editor/GLTFImporter/Source/GLTFImporter/GLTFImporter.Build.cs
+++ b/Engine/Plugins/Editor/GLTFImporter/Source/GLTFImporter/GLTFImporter.Build.cs
@@ -30,6 +30,7 @@ namespace UnrealBuildTool.Rules
                     "Engine",
                     "UnrealEd",
                     "MeshDescription",
+					"StaticMeshDescription",
                     "MeshUtilities",
                     "MessageLog",
                     "Json",
diff --git a/Engine/Plugins/Editor/GLTFImporter/Source/GLTFImporter/Private/GLTFStaticMeshFactory.cpp b/Engine/Plugins/Editor/GLTFImporter/Source/GLTFImporter/Private/GLTFStaticMeshFactory.cpp
--- a/Engine/Plugins/Editor/GLTFImporter/Source/GLTFImporter/Private/GLTFStaticMeshFactory.cpp
+++ b/Engine/Plugins/Editor/GLTFImporter/Source/GLTFImporter/Private/GLTFStaticMeshFactory.cpp
@@ -6,9 +6,7 @@
 
 #include "AssetRegistryModule.h"
 #include "Engine/StaticMesh.h"
-#include "MeshAttributeArray.h"
-#include "MeshAttributes.h"
-#include "MeshDescription.h"
+#include "StaticMeshAttributes.h"
 #include "PackageTools.h"
 
 namespace GLTF
@@ -182,7 +180,6 @@ namespace GLTF
 		// GLTF currently only supports LODs via MSFT_lod, for now use always 0
 		const int32 LODIndex = 0;
 		FMeshDescription* MeshDescription = StaticMesh->CreateMeshDescription(LODIndex);
-		StaticMesh->RegisterMeshAttributes(*MeshDescription);
 
 		FillMeshDescription(Mesh, MeshDescription);
 
@@ -218,23 +215,16 @@ namespace GLTF
 	{
 		const int32 NumUVs = FMath::Max(1, GetNumUVs(Mesh));
 
-		TVertexAttributesRef<FVector> VertexPositions =
-			MeshDescription->VertexAttributes().GetAttributesRef<FVector>(MeshAttribute::Vertex::Position);
-		TEdgeAttributesRef<bool>  EdgeHardnesses = MeshDescription->EdgeAttributes().GetAttributesRef<bool>(MeshAttribute::Edge::IsHard);
-		TEdgeAttributesRef<float> EdgeCreaseSharpnesses =
-			MeshDescription->EdgeAttributes().GetAttributesRef<float>(MeshAttribute::Edge::CreaseSharpness);
-		TPolygonGroupAttributesRef<FName> PolygonGroupImportedMaterialSlotNames =
-			MeshDescription->PolygonGroupAttributes().GetAttributesRef<FName>(MeshAttribute::PolygonGroup::ImportedMaterialSlotName);
-		TVertexInstanceAttributesRef<FVector> VertexInstanceNormals =
-			MeshDescription->VertexInstanceAttributes().GetAttributesRef<FVector>(MeshAttribute::VertexInstance::Normal);
-		TVertexInstanceAttributesRef<FVector> VertexInstanceTangents =
-			MeshDescription->VertexInstanceAttributes().GetAttributesRef<FVector>(MeshAttribute::VertexInstance::Tangent);
-		TVertexInstanceAttributesRef<float> VertexInstanceBinormalSigns =
-			MeshDescription->VertexInstanceAttributes().GetAttributesRef<float>(MeshAttribute::VertexInstance::BinormalSign);
-		TVertexInstanceAttributesRef<FVector2D> VertexInstanceUVs =
-			MeshDescription->VertexInstanceAttributes().GetAttributesRef<FVector2D>(MeshAttribute::VertexInstance::TextureCoordinate);
-		TVertexInstanceAttributesRef<FVector4> VertexInstanceColors =
-			MeshDescription->VertexInstanceAttributes().GetAttributesRef<FVector4>(MeshAttribute::VertexInstance::Color);
+		FStaticMeshAttributes StaticMeshAttributes(*MeshDescription);
+		TVertexAttributesRef<FVector> VertexPositions = StaticMeshAttributes.GetVertexPositions();
+		TEdgeAttributesRef<bool>  EdgeHardnesses = StaticMeshAttributes.GetEdgeHardnesses();
+		TEdgeAttributesRef<float> EdgeCreaseSharpnesses = StaticMeshAttributes.GetEdgeCreaseSharpnesses();
+		TPolygonGroupAttributesRef<FName> PolygonGroupImportedMaterialSlotNames = StaticMeshAttributes.GetPolygonGroupMaterialSlotNames();
+		TVertexInstanceAttributesRef<FVector> VertexInstanceNormals = StaticMeshAttributes.GetVertexInstanceNormals();
+		TVertexInstanceAttributesRef<FVector> VertexInstanceTangents = StaticMeshAttributes.GetVertexInstanceTangents();
+		TVertexInstanceAttributesRef<float> VertexInstanceBinormalSigns = StaticMeshAttributes.GetVertexInstanceBinormalSigns();
+		TVertexInstanceAttributesRef<FVector2D> VertexInstanceUVs = StaticMeshAttributes.GetVertexInstanceUVs();
+		TVertexInstanceAttributesRef<FVector4> VertexInstanceColors = StaticMeshAttributes.GetVertexInstanceColors();
 		VertexInstanceUVs.SetNumIndices(NumUVs);
 
 		MaterialIndicesUsed.Empty(10);
diff --git a/Engine/Plugins/Editor/MeshEditor/Source/MeshEditor/MeshEditor.Build.cs b/Engine/Plugins/Editor/MeshEditor/Source/MeshEditor/MeshEditor.Build.cs
--- a/Engine/Plugins/Editor/MeshEditor/Source/MeshEditor/MeshEditor.Build.cs
+++ b/Engine/Plugins/Editor/MeshEditor/Source/MeshEditor/MeshEditor.Build.cs
@@ -20,6 +20,7 @@ namespace UnrealBuildTool.Rules
                     "RenderCore",
 					"EditableMesh",
                     "MeshDescription",
+					"StaticMeshDescription",
 					"ViewportInteraction",
                     "VREditor",
 					"Projects",
diff --git a/Engine/Plugins/Editor/MeshEditor/Source/MeshEditor/MeshEditorGeometryCollectionAdapter.cpp b/Engine/Plugins/Editor/MeshEditor/Source/MeshEditor/MeshEditorGeometryCollectionAdapter.cpp
--- a/Engine/Plugins/Editor/MeshEditor/Source/MeshEditor/MeshEditorGeometryCollectionAdapter.cpp
+++ b/Engine/Plugins/Editor/MeshEditor/Source/MeshEditor/MeshEditorGeometryCollectionAdapter.cpp
@@ -2,7 +2,7 @@
 
 #include "MeshEditorGeometryCollectionAdapter.h"
 #include "EditableMesh.h"
-#include "MeshAttributes.h"
+#include "StaticMeshAttributes.h"
 #include "WireframeMeshComponent.h"
 #include "EditorSupportDelegates.h"
 #include "GeometryCollection/GeometryCollection.h"
@@ -11,7 +11,7 @@
 
 namespace GeometryCollectionAdapter
 {
-	static FColor GetEdgeColor(bool bIsHardEdge, bool bIsUVSeam)
+	static FColor GetEdgeColor(bool bIsHardEdge)
 	{
 		return FColor::White;
 	}
@@ -90,12 +90,11 @@ void UMeshEditorGeometryCollectionAdapter::OnRebuildRenderMesh(const UEditableMe
 
 			// Add all edges
 			TEdgeAttributesConstRef<bool> EdgeHardnesses = MeshDescription->EdgeAttributes().GetAttributesRef<bool>(MeshAttribute::Edge::IsHard);
-			TEdgeAttributesConstRef<bool> EdgeUVSeams = MeshDescription->EdgeAttributes().GetAttributesRef<bool>(MeshAttribute::Edge::IsUVSeam);
 			for (const FEdgeID EdgeID : MeshDescription->Edges().GetElementIDs())
 			{
 				WireframeMesh->AddEdge(EdgeID);
 				WireframeMesh->SetEdgeVertices(EdgeID, EditableMesh->GetEdgeVertex(EdgeID, 0), EditableMesh->GetEdgeVertex(EdgeID, 1));
-				WireframeMesh->SetEdgeColor(EdgeID, GeometryCollectionAdapter::GetEdgeColor(EdgeHardnesses[EdgeID], EdgeUVSeams[EdgeID]));
+				WireframeMesh->SetEdgeColor(EdgeID, GeometryCollectionAdapter::GetEdgeColor(EdgeHardnesses[EdgeID]));
 			}
 
 			for (int Element = 0; Element < GeometryCollection->NumElements(FGeometryCollection::TransformGroup); Element++)
@@ -306,13 +305,12 @@ void UMeshEditorGeometryCollectionAdapter::OnCreateEdges( const UEditableMesh* E
 	check( MeshDescription );
 
 	TEdgeAttributesConstRef<bool> EdgeHardnesses = MeshDescription->EdgeAttributes().GetAttributesRef<bool>( MeshAttribute::Edge::IsHard );
-	TEdgeAttributesConstRef<bool> EdgeUVSeams = MeshDescription->EdgeAttributes().GetAttributesRef<bool>( MeshAttribute::Edge::IsUVSeam );
 
 	for( const FEdgeID EdgeID : EdgeIDs )
 	{
 		WireframeMesh->AddEdge( EdgeID );
 		WireframeMesh->SetEdgeVertices( EdgeID, EditableMesh->GetEdgeVertex( EdgeID, 0 ), EditableMesh->GetEdgeVertex( EdgeID, 1 ) );
-		WireframeMesh->SetEdgeColor( EdgeID, GeometryCollectionAdapter::GetEdgeColor( EdgeHardnesses[ EdgeID ], EdgeUVSeams[ EdgeID ] ) );
+		WireframeMesh->SetEdgeColor( EdgeID, GeometryCollectionAdapter::GetEdgeColor( EdgeHardnesses[ EdgeID ] ) );
 
 		for( const FPolygonID PolygonID : MeshDescription->GetEdgeConnectedPolygons( EdgeID ) )
 		{
@@ -363,12 +361,10 @@ void UMeshEditorGeometryCollectionAdapter::OnSetEdgeAttribute( const UEditableMe
 	check( MeshDescription );
 
 	TEdgeAttributesConstRef<bool> EdgeHardnesses = MeshDescription->EdgeAttributes().GetAttributesRef<bool>( MeshAttribute::Edge::IsHard );
-	TEdgeAttributesConstRef<bool> EdgeUVSeams = MeshDescription->EdgeAttributes().GetAttributesRef<bool>( MeshAttribute::Edge::IsUVSeam );
 
-	if( Attribute.AttributeName == MeshAttribute::Edge::IsHard ||
-		Attribute.AttributeName == MeshAttribute::Edge::IsUVSeam )
+	if( Attribute.AttributeName == MeshAttribute::Edge::IsHard )
 	{
-		WireframeMesh->SetEdgeColor( EdgeID, GeometryCollectionAdapter::GetEdgeColor( EdgeHardnesses[ EdgeID ], EdgeUVSeams[ EdgeID ] ) );
+		WireframeMesh->SetEdgeColor( EdgeID, GeometryCollectionAdapter::GetEdgeColor( EdgeHardnesses[ EdgeID ] ) );
 	}
 }
 
diff --git a/Engine/Plugins/Editor/MeshEditor/Source/MeshEditor/MeshEditorMode.cpp b/Engine/Plugins/Editor/MeshEditor/Source/MeshEditor/MeshEditorMode.cpp
--- a/Engine/Plugins/Editor/MeshEditor/Source/MeshEditor/MeshEditorMode.cpp
+++ b/Engine/Plugins/Editor/MeshEditor/Source/MeshEditor/MeshEditorMode.cpp
@@ -6,7 +6,7 @@
 #include "MeshEditorModeToolkit.h"
 #include "MeshEditorUtilities.h"
 #include "EditableMesh.h"
-#include "MeshAttributes.h"
+#include "StaticMeshAttributes.h"
 #include "EditableMeshFactory.h"
 #include "MeshElement.h"
 #include "DynamicMeshBuilder.h"
@@ -46,7 +46,7 @@
 #include "Framework/Application/SlateApplication.h"
 #include "Materials/Material.h"
 #include "Components/PrimitiveComponent.h"
-#include "ILevelViewport.h"
+#include "IAssetViewport.h"
 #include "SLevelViewport.h"
 #include "MeshEditorSelectionModifiers.h"
 #include "Algo/Find.h"
@@ -960,10 +960,11 @@ void FMeshEditorMode::Enter()
 	{
 		const TSharedRef<ILevelEditor>& LevelEditor = LevelEditorModule.GetFirstLevelEditor().ToSharedRef();
 
+
 		// Do we have an active perspective viewport that is valid for VR?  If so, go ahead and use that.
 		TSharedPtr<FEditorViewportClient> ViewportClient;
 		{
-			TSharedPtr<ILevelViewport> ActiveLevelViewport = LevelEditor->GetActiveViewportInterface();
+			TSharedPtr<IAssetViewport> ActiveLevelViewport = LevelEditor->GetActiveViewportInterface();
 			if (ActiveLevelViewport.IsValid())
 			{
 				ViewportClient = StaticCastSharedRef<SLevelViewport>(ActiveLevelViewport->AsWidget())->GetViewportClient();
diff --git a/Engine/Plugins/Editor/MeshEditor/Source/MeshEditor/MeshEditorSelectionModifiers.cpp b/Engine/Plugins/Editor/MeshEditor/Source/MeshEditor/MeshEditorSelectionModifiers.cpp
--- a/Engine/Plugins/Editor/MeshEditor/Source/MeshEditor/MeshEditorSelectionModifiers.cpp
+++ b/Engine/Plugins/Editor/MeshEditor/Source/MeshEditor/MeshEditorSelectionModifiers.cpp
@@ -4,7 +4,7 @@
 
 #include "EditableMesh.h"
 #include "IMeshEditorModeUIContract.h"
-#include "MeshAttributes.h"
+#include "StaticMeshAttributes.h"
 #include "MeshEditorStyle.h"
 #include "UObject/UObjectIterator.h"
 
diff --git a/Engine/Plugins/Editor/MeshEditor/Source/MeshEditor/MeshEditorStaticMeshAdapter.cpp b/Engine/Plugins/Editor/MeshEditor/Source/MeshEditor/MeshEditorStaticMeshAdapter.cpp
--- a/Engine/Plugins/Editor/MeshEditor/Source/MeshEditor/MeshEditorStaticMeshAdapter.cpp
+++ b/Engine/Plugins/Editor/MeshEditor/Source/MeshEditor/MeshEditorStaticMeshAdapter.cpp
@@ -2,20 +2,13 @@
 
 #include "MeshEditorStaticMeshAdapter.h"
 #include "EditableMesh.h"
-#include "MeshAttributes.h"
+#include "StaticMeshAttributes.h"
 #include "WireframeMeshComponent.h"
 
 
-static FColor GetEdgeColor( bool bIsHardEdge, bool bIsUVSeam )
+static FColor GetEdgeColor( bool bIsHardEdge )
 {
-	if( bIsUVSeam )
-	{
-		return bIsHardEdge ? FColor( 20, 80, 25 ) : FColor( 0, 46, 10 );
-	}
-	else
-	{
-		return bIsHardEdge ? FColor( 20, 20, 128 ) : FColor( 0, 0, 76 );
-	}
+	return bIsHardEdge ? FColor( 20, 20, 128 ) : FColor( 0, 0, 76 );
 }
 
 
@@ -68,12 +61,11 @@ void UMeshEditorStaticMeshAdapter::OnRebuildRenderMesh( const UEditableMesh* Edi
 
 	// Add all edges
 	TEdgeAttributesConstRef<bool> EdgeHardnesses = MeshDescription->EdgeAttributes().GetAttributesRef<bool>( MeshAttribute::Edge::IsHard );
-	TEdgeAttributesConstRef<bool> EdgeUVSeams = MeshDescription->EdgeAttributes().GetAttributesRef<bool>( MeshAttribute::Edge::IsUVSeam );
 	for( const FEdgeID EdgeID : MeshDescription->Edges().GetElementIDs() )
 	{
 		WireframeMesh->AddEdge( EdgeID );
 		WireframeMesh->SetEdgeVertices( EdgeID, EditableMesh->GetEdgeVertex( EdgeID, 0 ), EditableMesh->GetEdgeVertex( EdgeID, 1 ) );
-		WireframeMesh->SetEdgeColor( EdgeID, GetEdgeColor( EdgeHardnesses[ EdgeID ], EdgeUVSeams[ EdgeID ] ) );
+		WireframeMesh->SetEdgeColor( EdgeID, GetEdgeColor( EdgeHardnesses[ EdgeID ] ) );
 	}
 
 	// Add all polygons and edge instances
@@ -231,13 +223,12 @@ void UMeshEditorStaticMeshAdapter::OnCreateEdges( const UEditableMesh* EditableM
 	check( MeshDescription );
 
 	TEdgeAttributesConstRef<bool> EdgeHardnesses = MeshDescription->EdgeAttributes().GetAttributesRef<bool>( MeshAttribute::Edge::IsHard );
-	TEdgeAttributesConstRef<bool> EdgeUVSeams = MeshDescription->EdgeAttributes().GetAttributesRef<bool>( MeshAttribute::Edge::IsUVSeam );
 
 	for( const FEdgeID EdgeID : EdgeIDs )
 	{
 		WireframeMesh->AddEdge( EdgeID );
 		WireframeMesh->SetEdgeVertices( EdgeID, EditableMesh->GetEdgeVertex( EdgeID, 0 ), EditableMesh->GetEdgeVertex( EdgeID, 1 ) );
-		WireframeMesh->SetEdgeColor( EdgeID, GetEdgeColor( EdgeHardnesses[ EdgeID ], EdgeUVSeams[ EdgeID ] ) );
+		WireframeMesh->SetEdgeColor( EdgeID, GetEdgeColor( EdgeHardnesses[ EdgeID ] ) );
 
 		for( const FPolygonID PolygonID : MeshDescription->GetEdgeConnectedPolygons( EdgeID ) )
 		{
@@ -279,12 +270,10 @@ void UMeshEditorStaticMeshAdapter::OnSetEdgeAttribute( const UEditableMesh* Edit
 	check( MeshDescription );
 
 	TEdgeAttributesConstRef<bool> EdgeHardnesses = MeshDescription->EdgeAttributes().GetAttributesRef<bool>( MeshAttribute::Edge::IsHard );
-	TEdgeAttributesConstRef<bool> EdgeUVSeams = MeshDescription->EdgeAttributes().GetAttributesRef<bool>( MeshAttribute::Edge::IsUVSeam );
 
-	if( Attribute.AttributeName == MeshAttribute::Edge::IsHard ||
-		Attribute.AttributeName == MeshAttribute::Edge::IsUVSeam )
+	if( Attribute.AttributeName == MeshAttribute::Edge::IsHard )
 	{
-		WireframeMesh->SetEdgeColor( EdgeID, GetEdgeColor( EdgeHardnesses[ EdgeID ], EdgeUVSeams[ EdgeID ] ) );
+		WireframeMesh->SetEdgeColor( EdgeID, GetEdgeColor( EdgeHardnesses[ EdgeID ] ) );
 	}
 }
 
diff --git a/Engine/Plugins/Editor/MeshEditor/Source/MeshEditor/MeshEditorUtilities.cpp b/Engine/Plugins/Editor/MeshEditor/Source/MeshEditor/MeshEditorUtilities.cpp
--- a/Engine/Plugins/Editor/MeshEditor/Source/MeshEditor/MeshEditorUtilities.cpp
+++ b/Engine/Plugins/Editor/MeshEditor/Source/MeshEditor/MeshEditorUtilities.cpp
@@ -3,7 +3,7 @@
 #include "MeshEditorUtilities.h"
 
 #include "EditableMesh.h"
-#include "MeshAttributes.h"
+#include "StaticMeshAttributes.h"
 #include "MeshElement.h"
 #include "Materials/Material.h"
 
diff --git a/Engine/Plugins/Editor/MeshEditor/Source/PolygonModeling/CollapseAllHierarchyCommand.cpp b/Engine/Plugins/Editor/MeshEditor/Source/PolygonModeling/CollapseAllHierarchyCommand.cpp
--- a/Engine/Plugins/Editor/MeshEditor/Source/PolygonModeling/CollapseAllHierarchyCommand.cpp
+++ b/Engine/Plugins/Editor/MeshEditor/Source/PolygonModeling/CollapseAllHierarchyCommand.cpp
@@ -5,8 +5,7 @@
 #include "ScopedTransaction.h"
 #include "Editor.h"
 #include "Engine/Selection.h"
-#include "MeshAttributes.h"
-#include "MeshDescription.h"
+#include "StaticMeshAttributes.h"
 #include "PackageTools.h"
 #include "MeshFractureSettings.h"
 #include "EditableMeshFactory.h"
diff --git a/Engine/Plugins/Editor/MeshEditor/Source/PolygonModeling/CollapseSelectedHierarchyCommand.cpp b/Engine/Plugins/Editor/MeshEditor/Source/PolygonModeling/CollapseSelectedHierarchyCommand.cpp
--- a/Engine/Plugins/Editor/MeshEditor/Source/PolygonModeling/CollapseSelectedHierarchyCommand.cpp
+++ b/Engine/Plugins/Editor/MeshEditor/Source/PolygonModeling/CollapseSelectedHierarchyCommand.cpp
@@ -5,8 +5,7 @@
 #include "ScopedTransaction.h"
 #include "Editor.h"
 #include "Engine/Selection.h"
-#include "MeshAttributes.h"
-#include "MeshDescription.h"
+#include "StaticMeshAttributes.h"
 #include "PackageTools.h"
 #include "MeshFractureSettings.h"
 #include "EditableMeshFactory.h"
diff --git a/Engine/Plugins/Editor/MeshEditor/Source/PolygonModeling/SplitMesh.cpp b/Engine/Plugins/Editor/MeshEditor/Source/PolygonModeling/SplitMesh.cpp
--- a/Engine/Plugins/Editor/MeshEditor/Source/PolygonModeling/SplitMesh.cpp
+++ b/Engine/Plugins/Editor/MeshEditor/Source/PolygonModeling/SplitMesh.cpp
@@ -7,8 +7,7 @@
 #include "Editor.h"
 #include "Engine/Selection.h"
 #include "IMeshEditorModeEditingContract.h"
-#include "MeshAttributes.h"
-#include "MeshDescription.h"
+#include "StaticMeshAttributes.h"
 #include "MeshEditor/MeshEditorMode.h"
 #include "PackageTools.h"
 #include "Settings/EditorExperimentalSettings.h"
@@ -322,7 +321,6 @@ void USplitMeshCommand::Execute(IMeshEditorModeEditingContract& MeshEditorMode)
 
 		FElementIDRemappings Remappings;
 		NewMeshDescription->Compact(Remappings);
-		NewMeshDescription->TriangulateMesh();
 		for(int32 PolygonIdIndex = 0; PolygonIdIndex <= GroupId.GetValue(); ++PolygonIdIndex)
 		{
 			NewStaticMesh->StaticMaterials.Add(UMaterial::GetDefaultMaterial(MD_Surface));
diff --git a/Engine/Plugins/Editor/SpeedTreeImporter/Source/SpeedTreeImporter/Private/SpeedTreeImportFactory.cpp b/Engine/Plugins/Editor/SpeedTreeImporter/Source/SpeedTreeImporter/Private/SpeedTreeImportFactory.cpp
--- a/Engine/Plugins/Editor/SpeedTreeImporter/Source/SpeedTreeImporter/Private/SpeedTreeImportFactory.cpp
+++ b/Engine/Plugins/Editor/SpeedTreeImporter/Source/SpeedTreeImporter/Private/SpeedTreeImportFactory.cpp
@@ -57,9 +57,7 @@
 #include "AssetRegistryModule.h"
 #include "UnrealEd/Private/GeomFitUtils.h"
 #include "SpeedTreeWind.h"
-#include "MeshDescription.h"
-#include "MeshAttributes.h"
-#include "MeshAttributeArray.h"
+#include "StaticMeshAttributes.h"
 #include "ComponentReregisterContext.h"
 
 #if WITH_SPEEDTREE
@@ -1356,7 +1354,7 @@ FVertexInstanceID ProcessTriangleCorner(
 	const SpeedTree::SDrawCall* DrawCall,
 	const SpeedTree::st_uint32* Indices32,
 	const SpeedTree::st_uint16* Indices16,
-	FMeshDescription* MeshDescription,
+	FMeshDescription& MeshDescription,
 	const int32 IndexOffset,
 	const int32 NumUVs,
 	const SpeedTree::SRenderState* RenderState,
@@ -1376,7 +1374,7 @@ FVertexInstanceID ProcessTriangleCorner(
 
 	int32 VertexIndex = DrawCall->m_b32BitIndices ? Indices32[ Index ] : Indices16[ Index ];
 	FVertexID VertexID(IndexOffset + VertexIndex);
-	const FVertexInstanceID& VertexInstanceID = MeshDescription->CreateVertexInstance(VertexID);
+	const FVertexInstanceID& VertexInstanceID = MeshDescription.CreateVertexInstance(VertexID);
 
 	// tangents
 	DrawCall->GetProperty( SpeedTree::VERTEX_PROPERTY_NORMAL, VertexIndex, Data );
@@ -1650,22 +1648,22 @@ UObject* USpeedTreeImportFactory::FactoryCreateBinary7(UClass* InClass, UObject*
 						const SpeedTree::SLod* TreeLOD = &SpeedTreeGeometry->m_pLods[LODIndex];
 						FStaticMeshSourceModel& LODModel = StaticMesh->AddSourceModel();
 						FMeshDescription* MeshDescription = StaticMesh->CreateMeshDescription(LODIndex);
-						StaticMesh->RegisterMeshAttributes(*MeshDescription);
+						FStaticMeshAttributes Attributes(*MeshDescription);
 						
-						TVertexAttributesRef<FVector> VertexPositions = MeshDescription->VertexAttributes().GetAttributesRef<FVector>(MeshAttribute::Vertex::Position);
-						TEdgeAttributesRef<bool> EdgeHardnesses = MeshDescription->EdgeAttributes().GetAttributesRef<bool>(MeshAttribute::Edge::IsHard);
-						TEdgeAttributesRef<float> EdgeCreaseSharpnesses = MeshDescription->EdgeAttributes().GetAttributesRef<float>(MeshAttribute::Edge::CreaseSharpness);
-						TPolygonGroupAttributesRef<FName> PolygonGroupImportedMaterialSlotNames = MeshDescription->PolygonGroupAttributes().GetAttributesRef<FName>(MeshAttribute::PolygonGroup::ImportedMaterialSlotName);
-						TVertexInstanceAttributesRef<FVector> VertexInstanceNormals = MeshDescription->VertexInstanceAttributes().GetAttributesRef<FVector>(MeshAttribute::VertexInstance::Normal);
-						TVertexInstanceAttributesRef<FVector> VertexInstanceTangents = MeshDescription->VertexInstanceAttributes().GetAttributesRef<FVector>(MeshAttribute::VertexInstance::Tangent);
-						TVertexInstanceAttributesRef<float> VertexInstanceBinormalSigns = MeshDescription->VertexInstanceAttributes().GetAttributesRef<float>(MeshAttribute::VertexInstance::BinormalSign);
-						TVertexInstanceAttributesRef<FVector4> VertexInstanceColors = MeshDescription->VertexInstanceAttributes().GetAttributesRef<FVector4>(MeshAttribute::VertexInstance::Color);
-						TVertexInstanceAttributesRef<FVector2D> VertexInstanceUVs = MeshDescription->VertexInstanceAttributes().GetAttributesRef<FVector2D>(MeshAttribute::VertexInstance::TextureCoordinate);
+						TVertexAttributesRef<FVector> VertexPositions = Attributes.GetVertexPositions();
+						TEdgeAttributesRef<bool> EdgeHardnesses = Attributes.GetEdgeHardnesses();
+						TEdgeAttributesRef<float> EdgeCreaseSharpnesses = Attributes.GetEdgeCreaseSharpnesses();
+						TPolygonGroupAttributesRef<FName> PolygonGroupImportedMaterialSlotNames = Attributes.GetPolygonGroupMaterialSlotNames();
+						TVertexInstanceAttributesRef<FVector> VertexInstanceNormals = Attributes.GetVertexInstanceNormals();
+						TVertexInstanceAttributesRef<FVector> VertexInstanceTangents = Attributes.GetVertexInstanceTangents();
+						TVertexInstanceAttributesRef<float> VertexInstanceBinormalSigns = Attributes.GetVertexInstanceBinormalSigns();
+						TVertexInstanceAttributesRef<FVector4> VertexInstanceColors = Attributes.GetVertexInstanceColors();
+						TVertexInstanceAttributesRef<FVector2D> VertexInstanceUVs = Attributes.GetVertexInstanceUVs();
 
 						TMap<int32, FPolygonGroupID> MaterialToPolygonGroup;
 						for (int32 MatIndex = 0; MatIndex < StaticMesh->StaticMaterials.Num(); ++MatIndex)
 						{
-							const FPolygonGroupID& PolygonGroupID = MeshDescription->CreatePolygonGroup();
+							const FPolygonGroupID PolygonGroupID = MeshDescription->CreatePolygonGroup();
 							PolygonGroupImportedMaterialSlotNames[PolygonGroupID] = StaticMesh->StaticMaterials[MatIndex].ImportedMaterialSlotName;
 							MaterialToPolygonGroup.Add(MatIndex, PolygonGroupID);
 						}
@@ -1743,7 +1741,7 @@ UObject* USpeedTreeImportFactory::FactoryCreateBinary7(UClass* InClass, UObject*
 								
 								RenderStateIndexToStaticMeshIndex.Add(DrawCall->m_nRenderStateIndex, StaticMesh->StaticMaterials.Num());
 								MaterialIndex = StaticMesh->StaticMaterials.Add(FStaticMaterial(Material, FName(*MaterialName), FName(*MaterialName)));
-								const FPolygonGroupID& PolygonGroupID = MeshDescription->CreatePolygonGroup();
+								const FPolygonGroupID PolygonGroupID = MeshDescription->CreatePolygonGroup();
 								PolygonGroupImportedMaterialSlotNames[PolygonGroupID] = StaticMesh->StaticMaterials[MaterialIndex].ImportedMaterialSlotName;
 								MaterialToPolygonGroup.Add(MaterialIndex, PolygonGroupID);
 							}
@@ -1789,7 +1787,7 @@ UObject* USpeedTreeImportFactory::FactoryCreateBinary7(UClass* InClass, UObject*
 								for (int32 Corner = 0; Corner < 3; ++Corner)
 								{
 									//Create the vertex instances
-									CornerVertexInstanceIDs[Corner] = ProcessTriangleCorner( SpeedTree, TriangleIndex, Corner, DrawCall, Indices32, Indices16, MeshDescription, IndexOffset, NumUVs, RenderState,
+									CornerVertexInstanceIDs[Corner] = ProcessTriangleCorner( SpeedTree, TriangleIndex, Corner, DrawCall, Indices32, Indices16, *MeshDescription, IndexOffset, NumUVs, RenderState,
 										VertexInstanceNormals, VertexInstanceTangents, VertexInstanceBinormalSigns, VertexInstanceColors, VertexInstanceUVs);
 
 									CornerVertexIDs[Corner] = MeshDescription->GetVertexInstanceVertex(CornerVertexInstanceIDs[Corner]);
@@ -1825,24 +1823,24 @@ UObject* USpeedTreeImportFactory::FactoryCreateBinary7(UClass* InClass, UObject*
 					FStaticMeshSourceModel& LODModel = StaticMesh->AddSourceModel();
 					const int32 LODIndex = StaticMesh->GetNumSourceModels() - 1;
 					FMeshDescription* MeshDescription = StaticMesh->CreateMeshDescription(LODIndex);
-					StaticMesh->RegisterMeshAttributes(*MeshDescription);
-
-					TVertexAttributesRef<FVector> VertexPositions = MeshDescription->VertexAttributes().GetAttributesRef<FVector>(MeshAttribute::Vertex::Position);
-					TEdgeAttributesRef<bool> EdgeHardnesses = MeshDescription->EdgeAttributes().GetAttributesRef<bool>(MeshAttribute::Edge::IsHard);
-					TEdgeAttributesRef<float> EdgeCreaseSharpnesses = MeshDescription->EdgeAttributes().GetAttributesRef<float>(MeshAttribute::Edge::CreaseSharpness);
-					TPolygonGroupAttributesRef<FName> PolygonGroupImportedMaterialSlotNames = MeshDescription->PolygonGroupAttributes().GetAttributesRef<FName>(MeshAttribute::PolygonGroup::ImportedMaterialSlotName);
-					TVertexInstanceAttributesRef<FVector> VertexInstanceNormals = MeshDescription->VertexInstanceAttributes().GetAttributesRef<FVector>(MeshAttribute::VertexInstance::Normal);
-					TVertexInstanceAttributesRef<FVector> VertexInstanceTangents = MeshDescription->VertexInstanceAttributes().GetAttributesRef<FVector>(MeshAttribute::VertexInstance::Tangent);
-					TVertexInstanceAttributesRef<float> VertexInstanceBinormalSigns = MeshDescription->VertexInstanceAttributes().GetAttributesRef<float>(MeshAttribute::VertexInstance::BinormalSign);
-					TVertexInstanceAttributesRef<FVector4> VertexInstanceColors = MeshDescription->VertexInstanceAttributes().GetAttributesRef<FVector4>(MeshAttribute::VertexInstance::Color);
-					TVertexInstanceAttributesRef<FVector2D> VertexInstanceUVs = MeshDescription->VertexInstanceAttributes().GetAttributesRef<FVector2D>(MeshAttribute::VertexInstance::TextureCoordinate);
+					FStaticMeshAttributes Attributes(*MeshDescription);
+
+					TVertexAttributesRef<FVector> VertexPositions = Attributes.GetVertexPositions();
+					TEdgeAttributesRef<bool> EdgeHardnesses = Attributes.GetEdgeHardnesses();
+					TEdgeAttributesRef<float> EdgeCreaseSharpnesses = Attributes.GetEdgeCreaseSharpnesses();
+					TPolygonGroupAttributesRef<FName> PolygonGroupImportedMaterialSlotNames = Attributes.GetPolygonGroupMaterialSlotNames();
+					TVertexInstanceAttributesRef<FVector> VertexInstanceNormals = Attributes.GetVertexInstanceNormals();
+					TVertexInstanceAttributesRef<FVector> VertexInstanceTangents = Attributes.GetVertexInstanceTangents();
+					TVertexInstanceAttributesRef<float> VertexInstanceBinormalSigns = Attributes.GetVertexInstanceBinormalSigns();
+					TVertexInstanceAttributesRef<FVector4> VertexInstanceColors = Attributes.GetVertexInstanceColors();
+					TVertexInstanceAttributesRef<FVector2D> VertexInstanceUVs = Attributes.GetVertexInstanceUVs();
 
 					//Speedtree use UVs to store is data
 					VertexInstanceUVs.SetNumIndices(2);
 
 					for (int32 MatIndex = 0; MatIndex < StaticMesh->StaticMaterials.Num(); ++MatIndex)
 					{
-						const FPolygonGroupID& PolygonGroupID = MeshDescription->CreatePolygonGroup();
+						const FPolygonGroupID PolygonGroupID = MeshDescription->CreatePolygonGroup();
 						PolygonGroupImportedMaterialSlotNames[PolygonGroupID] = StaticMesh->StaticMaterials[MatIndex].ImportedMaterialSlotName;
 					}
 
@@ -1850,7 +1848,7 @@ UObject* USpeedTreeImportFactory::FactoryCreateBinary7(UClass* InClass, UObject*
 					UMaterialInterface* Material = CreateSpeedTreeMaterial7(InParent, MaterialName, &SpeedTreeGeometry->m_aBillboardRenderStates[SpeedTree::RENDER_PASS_MAIN], Options, WindType, SpeedTreeGeometry->m_sVertBBs.m_nNumBillboards, LoadedPackages);
 					int32 MaterialIndex = StaticMesh->StaticMaterials.Add(FStaticMaterial(Material, FName(*MaterialName), FName(*MaterialName)));
 
-					const FPolygonGroupID& CurrentPolygonGroupID = MeshDescription->CreatePolygonGroup();
+					const FPolygonGroupID CurrentPolygonGroupID = MeshDescription->CreatePolygonGroup();
 					PolygonGroupImportedMaterialSlotNames[CurrentPolygonGroupID] = StaticMesh->StaticMaterials[MaterialIndex].ImportedMaterialSlotName;
 					
 					// fill out triangles
@@ -2154,17 +2152,17 @@ UObject* USpeedTreeImportFactory::FactoryCreateBinary8(UClass* InClass, UObject*
 			SpeedTree8::CLod LOD = SpeedTree.Lods()[LODIndex];
 			
 			FMeshDescription* MeshDescription = StaticMesh->CreateMeshDescription(LODIndex);
-			StaticMesh->RegisterMeshAttributes(*MeshDescription);
-
-			TVertexAttributesRef<FVector> VertexPositions = MeshDescription->VertexAttributes().GetAttributesRef<FVector>(MeshAttribute::Vertex::Position);
-			TEdgeAttributesRef<bool> EdgeHardnesses = MeshDescription->EdgeAttributes().GetAttributesRef<bool>(MeshAttribute::Edge::IsHard);
-			TEdgeAttributesRef<float> EdgeCreaseSharpnesses = MeshDescription->EdgeAttributes().GetAttributesRef<float>(MeshAttribute::Edge::CreaseSharpness);
-			TPolygonGroupAttributesRef<FName> PolygonGroupImportedMaterialSlotNames = MeshDescription->PolygonGroupAttributes().GetAttributesRef<FName>(MeshAttribute::PolygonGroup::ImportedMaterialSlotName);
-			TVertexInstanceAttributesRef<FVector> VertexInstanceNormals = MeshDescription->VertexInstanceAttributes().GetAttributesRef<FVector>(MeshAttribute::VertexInstance::Normal);
-			TVertexInstanceAttributesRef<FVector> VertexInstanceTangents = MeshDescription->VertexInstanceAttributes().GetAttributesRef<FVector>(MeshAttribute::VertexInstance::Tangent);
-			TVertexInstanceAttributesRef<float> VertexInstanceBinormalSigns = MeshDescription->VertexInstanceAttributes().GetAttributesRef<float>(MeshAttribute::VertexInstance::BinormalSign);
-			TVertexInstanceAttributesRef<FVector4> VertexInstanceColors = MeshDescription->VertexInstanceAttributes().GetAttributesRef<FVector4>(MeshAttribute::VertexInstance::Color);
-			TVertexInstanceAttributesRef<FVector2D> VertexInstanceUVs = MeshDescription->VertexInstanceAttributes().GetAttributesRef<FVector2D>(MeshAttribute::VertexInstance::TextureCoordinate);
+			FStaticMeshAttributes Attributes(*MeshDescription);
+
+			TVertexAttributesRef<FVector> VertexPositions = Attributes.GetVertexPositions();
+			TEdgeAttributesRef<bool> EdgeHardnesses = Attributes.GetEdgeHardnesses();
+			TEdgeAttributesRef<float> EdgeCreaseSharpnesses = Attributes.GetEdgeCreaseSharpnesses();
+			TPolygonGroupAttributesRef<FName> PolygonGroupImportedMaterialSlotNames = Attributes.GetPolygonGroupMaterialSlotNames();
+			TVertexInstanceAttributesRef<FVector> VertexInstanceNormals = Attributes.GetVertexInstanceNormals();
+			TVertexInstanceAttributesRef<FVector> VertexInstanceTangents = Attributes.GetVertexInstanceTangents();
+			TVertexInstanceAttributesRef<float> VertexInstanceBinormalSigns = Attributes.GetVertexInstanceBinormalSigns();
+			TVertexInstanceAttributesRef<FVector4> VertexInstanceColors = Attributes.GetVertexInstanceColors();
+			TVertexInstanceAttributesRef<FVector2D> VertexInstanceUVs = Attributes.GetVertexInstanceUVs();
 
 			//Speedtree use 8 UVs to store is data
 			VertexInstanceUVs.SetNumIndices(8);
diff --git a/Engine/Plugins/Editor/SpeedTreeImporter/Source/SpeedTreeImporter/SpeedTreeImporter.Build.cs b/Engine/Plugins/Editor/SpeedTreeImporter/Source/SpeedTreeImporter/SpeedTreeImporter.Build.cs
--- a/Engine/Plugins/Editor/SpeedTreeImporter/Source/SpeedTreeImporter/SpeedTreeImporter.Build.cs
+++ b/Engine/Plugins/Editor/SpeedTreeImporter/Source/SpeedTreeImporter/SpeedTreeImporter.Build.cs
@@ -35,6 +35,7 @@ namespace UnrealBuildTool.Rules
                     "UnrealEd",
                     "MainFrame",
                     "MeshDescription",
+					"StaticMeshDescription"
                 }
 				);
 
diff --git a/Engine/Plugins/Enterprise/DatasmithContent/Source/DatasmithContent/DatasmithContent.Build.cs b/Engine/Plugins/Enterprise/DatasmithContent/Source/DatasmithContent/DatasmithContent.Build.cs
--- a/Engine/Plugins/Enterprise/DatasmithContent/Source/DatasmithContent/DatasmithContent.Build.cs
+++ b/Engine/Plugins/Enterprise/DatasmithContent/Source/DatasmithContent/DatasmithContent.Build.cs
@@ -24,7 +24,8 @@ namespace UnrealBuildTool.Rules
 					"Landscape",
 					"LevelSequence",
 					"MeshDescription",
-					"Projects"
+					"Projects",
+					"StaticMeshDescription"
 				}
 			);
 
diff --git a/Engine/Plugins/Enterprise/DatasmithContent/Source/DatasmithContent/Private/ObjectTemplates/DatasmithStaticMeshTemplate.cpp b/Engine/Plugins/Enterprise/DatasmithContent/Source/DatasmithContent/Private/ObjectTemplates/DatasmithStaticMeshTemplate.cpp
--- a/Engine/Plugins/Enterprise/DatasmithContent/Source/DatasmithContent/Private/ObjectTemplates/DatasmithStaticMeshTemplate.cpp
+++ b/Engine/Plugins/Enterprise/DatasmithContent/Source/DatasmithContent/Private/ObjectTemplates/DatasmithStaticMeshTemplate.cpp
@@ -6,7 +6,7 @@
 
 #include "CoreTypes.h"
 #include "Engine/StaticMesh.h"
-#include "MeshDescription.h"
+#include "StaticMeshAttributes.h"
 #include "Templates/Casts.h"
 
 FDatasmithMeshBuildSettingsTemplate::FDatasmithMeshBuildSettingsTemplate()
@@ -262,11 +262,12 @@ UObject* UDatasmithStaticMeshTemplate::UpdateObject( UObject* Destination, bool
 	const int32 NumLODs = StaticMesh->GetNumLODs();
 	for ( int32 LODIndex = 0; LODIndex < NumLODs; ++LODIndex )
 	{
-		FMeshDescription* MeshDescription = StaticMesh->GetMeshDescription( LODIndex );
+		const FMeshDescription* MeshDescription = StaticMesh->GetMeshDescription( LODIndex );
 		if ( MeshDescription && MeshDescription->PolygonGroups().Num() == StaticMesh->StaticMaterials.Num() )
 		{
+			FStaticMeshConstAttributes Attributes(*MeshDescription);
 			TArray< FStaticMaterial > TempStaticMaterials;
-			TPolygonGroupAttributesConstRef< FName > PolygonGroupImportedMaterialSlotNames = MeshDescription->PolygonGroupAttributes().GetAttributesRef< FName >( MeshAttribute::PolygonGroup::ImportedMaterialSlotName );
+			TPolygonGroupAttributesConstRef< FName > PolygonGroupImportedMaterialSlotNames = Attributes.GetPolygonGroupMaterialSlotNames();
 			int32 SectionIndex = 0;
 			for ( const FPolygonGroupID PolygonGroupID : MeshDescription->PolygonGroups().GetElementIDs() )
 			{
diff --git a/Engine/Plugins/Experimental/AlembicImporter/Source/AlembicLibrary/Private/AbcImporter.cpp b/Engine/Plugins/Experimental/AlembicImporter/Source/AlembicLibrary/Private/AbcImporter.cpp
--- a/Engine/Plugins/Experimental/AlembicImporter/Source/AlembicLibrary/Private/AbcImporter.cpp
+++ b/Engine/Plugins/Experimental/AlembicImporter/Source/AlembicLibrary/Private/AbcImporter.cpp
@@ -25,9 +25,7 @@ PRAGMA_DEFAULT_VISIBILITY_END
 #include "Misc/ScopedSlowTask.h"
 
 #include "PackageTools.h"
-#include "MeshDescription.h"
-#include "MeshAttributes.h"
-#include "MeshAttributeArray.h"
+#include "StaticMeshAttributes.h"
 #include "MeshDescriptionOperations.h"
 #include "ObjectTools.h"
 
@@ -197,7 +195,6 @@ UStaticMesh* FAbcImporter::CreateStaticMeshFromSample(UObject* InParent, const F
 		int32 LODIndex = 0;
 		StaticMesh->AddSourceModel();
 		FMeshDescription* MeshDescription = StaticMesh->CreateMeshDescription(LODIndex);
-		StaticMesh->RegisterMeshAttributes(*MeshDescription);
 		// Generate a new lighting GUID (so its unique)
 		StaticMesh->LightingGuid = FGuid::NewGuid();
 
@@ -1281,22 +1278,29 @@ const uint32 FAbcImporter::GetNumMeshTracks() const
 
 void FAbcImporter::GenerateMeshDescriptionFromSample(const FAbcMeshSample* Sample, FMeshDescription* MeshDescription, UStaticMesh* StaticMesh)
 {
-	TVertexAttributesRef<FVector> VertexPositions = MeshDescription->VertexAttributes().GetAttributesRef<FVector>(MeshAttribute::Vertex::Position);
-	TEdgeAttributesRef<bool> EdgeHardnesses = MeshDescription->EdgeAttributes().GetAttributesRef<bool>(MeshAttribute::Edge::IsHard);
-	TEdgeAttributesRef<float> EdgeCreaseSharpnesses = MeshDescription->EdgeAttributes().GetAttributesRef<float>(MeshAttribute::Edge::CreaseSharpness);
-	TPolygonGroupAttributesRef<FName> PolygonGroupImportedMaterialSlotNames = MeshDescription->PolygonGroupAttributes().GetAttributesRef<FName>(MeshAttribute::PolygonGroup::ImportedMaterialSlotName);
-	TVertexInstanceAttributesRef<FVector> VertexInstanceNormals = MeshDescription->VertexInstanceAttributes().GetAttributesRef<FVector>(MeshAttribute::VertexInstance::Normal);
-	TVertexInstanceAttributesRef<FVector> VertexInstanceTangents = MeshDescription->VertexInstanceAttributes().GetAttributesRef<FVector>(MeshAttribute::VertexInstance::Tangent);
-	TVertexInstanceAttributesRef<float> VertexInstanceBinormalSigns = MeshDescription->VertexInstanceAttributes().GetAttributesRef<float>(MeshAttribute::VertexInstance::BinormalSign);
-	TVertexInstanceAttributesRef<FVector4> VertexInstanceColors = MeshDescription->VertexInstanceAttributes().GetAttributesRef<FVector4>(MeshAttribute::VertexInstance::Color);
-	TVertexInstanceAttributesRef<FVector2D> VertexInstanceUVs = MeshDescription->VertexInstanceAttributes().GetAttributesRef<FVector2D>(MeshAttribute::VertexInstance::TextureCoordinate);
+	if (MeshDescription == nullptr)
+	{
+		return;
+	}
+
+	FStaticMeshAttributes Attributes(*MeshDescription);
+
+	TVertexAttributesRef<FVector> VertexPositions = Attributes.GetVertexPositions();
+	TEdgeAttributesRef<bool> EdgeHardnesses = Attributes.GetEdgeHardnesses();
+	TEdgeAttributesRef<float> EdgeCreaseSharpnesses = Attributes.GetEdgeCreaseSharpnesses();
+	TPolygonGroupAttributesRef<FName> PolygonGroupImportedMaterialSlotNames = Attributes.GetPolygonGroupMaterialSlotNames();
+	TVertexInstanceAttributesRef<FVector> VertexInstanceNormals = Attributes.GetVertexInstanceNormals();
+	TVertexInstanceAttributesRef<FVector> VertexInstanceTangents = Attributes.GetVertexInstanceTangents();
+	TVertexInstanceAttributesRef<float> VertexInstanceBinormalSigns = Attributes.GetVertexInstanceBinormalSigns();
+	TVertexInstanceAttributesRef<FVector4> VertexInstanceColors = Attributes.GetVertexInstanceColors();
+	TVertexInstanceAttributesRef<FVector2D> VertexInstanceUVs = Attributes.GetVertexInstanceUVs();
 
 	//Speedtree use UVs to store is data
 	VertexInstanceUVs.SetNumIndices(Sample->NumUVSets);
 	
 	for (int32 MatIndex = 0; MatIndex < StaticMesh->StaticMaterials.Num(); ++MatIndex)
 	{
-		const FPolygonGroupID& PolygonGroupID = MeshDescription->CreatePolygonGroup();
+		const FPolygonGroupID PolygonGroupID = MeshDescription->CreatePolygonGroup();
 		PolygonGroupImportedMaterialSlotNames[PolygonGroupID] = StaticMesh->StaticMaterials[MatIndex].ImportedMaterialSlotName;
 	}
 
@@ -1320,7 +1324,7 @@ void FAbcImporter::GenerateMeshDescriptionFromSample(const FAbcMeshSample* Sampl
 			uint32 IndiceIndex = (TriangleIndex * 3) + Corner;
 			uint32 VertexIndex = Sample->Indices[IndiceIndex];
 			const FVertexID VertexID(VertexIndex);
-			const FVertexInstanceID& VertexInstanceID = MeshDescription->CreateVertexInstance(VertexID);
+			const FVertexInstanceID VertexInstanceID = MeshDescription->CreateVertexInstance(VertexID);
 
 			// tangents
 			FVector TangentX = Sample->TangentX[IndiceIndex];
diff --git a/Engine/Plugins/Experimental/BlastPlugin/Source/BlastAuthoring/Private/MeshUtility.cpp b/Engine/Plugins/Experimental/BlastPlugin/Source/BlastAuthoring/Private/MeshUtility.cpp
--- a/Engine/Plugins/Experimental/BlastPlugin/Source/BlastAuthoring/Private/MeshUtility.cpp
+++ b/Engine/Plugins/Experimental/BlastPlugin/Source/BlastAuthoring/Private/MeshUtility.cpp
@@ -2,7 +2,7 @@
 
 #include "MeshUtility.h"
 #include "PxPhysicsAPI.h"
-#include "MeshAttributes.h"
+#include "StaticMeshAttributes.h"
 #include "GeometryCollection/GeometryCollection.h"
 #include "GeometryCollection/GeometryCollectionClusteringUtility.h"
 #include "GeometryCollection/GeometryCollectionAlgo.h"
@@ -27,10 +27,11 @@ using namespace physx;
 void FMeshUtility::EditableMeshToBlastMesh(const UEditableMesh* SourceMesh, Nv::Blast::Mesh*& OutBlastMesh)
 {
 	const FMeshDescription* MeshDescription = SourceMesh->GetMeshDescription();
+	FStaticMeshConstAttributes Attributes(*MeshDescription);
 
-	TVertexAttributesConstRef<FVector> VertexPositions = MeshDescription->VertexAttributes().GetAttributesRef<FVector>(MeshAttribute::Vertex::Position);
-	TVertexInstanceAttributesConstRef<FVector2D> VertexUvs = MeshDescription->VertexInstanceAttributes().GetAttributesRef<FVector2D>(MeshAttribute::VertexInstance::TextureCoordinate);
-	TVertexInstanceAttributesConstRef<FVector> VertexNormals = MeshDescription->VertexInstanceAttributes().GetAttributesRef<FVector>(MeshAttribute::VertexInstance::Normal);
+	TVertexAttributesConstRef<FVector> VertexPositions = Attributes.GetVertexPositions();
+	TVertexInstanceAttributesConstRef<FVector2D> VertexUvs = Attributes.GetVertexInstanceUVs();
+	TVertexInstanceAttributesConstRef<FVector> VertexNormals = Attributes.GetVertexInstanceNormals();
 
 	// Blast representation
 	TArray<PxVec2> BlastUvs;
@@ -83,10 +84,11 @@ void FMeshUtility::EditableMeshToBlastMesh(const UEditableMesh* SourceMesh, int3
 {
 	OutBlastMesh = nullptr;
 	const FMeshDescription* MeshDescription = SourceMesh->GetMeshDescription();
+	FStaticMeshConstAttributes Attributes(*MeshDescription);
 
-	TVertexAttributesConstRef<FVector> VertexPositions = MeshDescription->VertexAttributes().GetAttributesRef<FVector>(MeshAttribute::Vertex::Position);
-	TVertexInstanceAttributesConstRef<FVector2D> VertexUVs = MeshDescription->VertexInstanceAttributes().GetAttributesRef<FVector2D>(MeshAttribute::VertexInstance::TextureCoordinate);
-	TVertexInstanceAttributesConstRef<FVector> VertexNormals = MeshDescription->VertexInstanceAttributes().GetAttributesRef<FVector>(MeshAttribute::VertexInstance::Normal);
+	TVertexAttributesConstRef<FVector> VertexPositions = Attributes.GetVertexPositions();
+	TVertexInstanceAttributesConstRef<FVector2D> VertexUVs = Attributes.GetVertexInstanceUVs();
+	TVertexInstanceAttributesConstRef<FVector> VertexNormals = Attributes.GetVertexInstanceNormals();
 	UGeometryCollection* GeometryCollection = Cast<UGeometryCollection>(static_cast<UObject*>(SourceMesh->GetSubMeshAddress().MeshObjectPtr));
 
 	//#define VALIDATE_INPUT
diff --git a/Engine/Plugins/Experimental/DataPrepEditor/Source/DataPrepLibraries/Private/DataPrepOperationsLibrary.cpp b/Engine/Plugins/Experimental/DataPrepEditor/Source/DataPrepLibraries/Private/DataPrepOperationsLibrary.cpp
--- a/Engine/Plugins/Experimental/DataPrepEditor/Source/DataPrepLibraries/Private/DataPrepOperationsLibrary.cpp
+++ b/Engine/Plugins/Experimental/DataPrepEditor/Source/DataPrepLibraries/Private/DataPrepOperationsLibrary.cpp
@@ -24,12 +24,10 @@
 #include "Materials/Material.h"
 #include "Materials/MaterialInterface.h"
 #include "Math/Vector2D.h"
-#include "MeshAttributeArray.h"
-#include "MeshDescription.h"
 #include "MeshDescriptionOperations.h"
-#include "MeshTypes.h"
 #include "Misc/FileHelper.h"
 #include "ObjectTools.h"
+#include "StaticMeshAttributes.h"
 #include "TessellationRendering.h"
 #include "UObject/SoftObjectPath.h"
 
@@ -241,8 +239,8 @@ namespace DataprepOperationsLibraryUtil
 		SourceModel.BuildSettings.bBuildReversedIndexBuffer = false;
 		SourceModel.BuildSettings.bComputeWeightedNormals = false;
 
-		FMeshDescription* StaticMeshDescription = SourceModel.MeshDescription.Get();
-		check( StaticMeshDescription );
+		FMeshDescription* MeshDescription = SourceModel.MeshDescription.Get();
+		check( MeshDescription );
 
 		// Create render data
 		StaticMesh->RenderData.Reset( new(FMemory::Malloc(sizeof(FStaticMeshRenderData)))FStaticMeshRenderData() );
@@ -361,9 +359,10 @@ void UDataprepOperationsLibrary::SetGenerateLightmapUVs( const TArray< UObject*
 			{
 				bDidChangeSettings |= SourceModel.BuildSettings.bGenerateLightmapUVs != bGenerateLightmapUVs;
 				SourceModel.BuildSettings.bGenerateLightmapUVs = bGenerateLightmapUVs;
-				if( FMeshDescription* MeshDescription = SourceModel.MeshDescription.Get() )
+				if( const FMeshDescription* MeshDescription = SourceModel.MeshDescription.Get() )
 				{
-					int32 UVChannelCount = MeshDescription->VertexInstanceAttributes().GetAttributesRef< FVector2D >( MeshAttribute::VertexInstance::TextureCoordinate ).GetNumIndices();
+					FStaticMeshConstAttributes Attributes(*MeshDescription);
+					int32 UVChannelCount = Attributes.GetVertexInstanceUVs().GetNumIndices();
 					MinBiggestUVChannel = FMath::Min( MinBiggestUVChannel, UVChannelCount - 1 );
 				}
 			}
diff --git a/Engine/Plugins/Experimental/DataPrepEditor/Source/DataPrepLibraries/Private/Fetchers/DataprepFloatFetcherLibrary.cpp b/Engine/Plugins/Experimental/DataPrepEditor/Source/DataPrepLibraries/Private/Fetchers/DataprepFloatFetcherLibrary.cpp
--- a/Engine/Plugins/Experimental/DataPrepEditor/Source/DataPrepLibraries/Private/Fetchers/DataprepFloatFetcherLibrary.cpp
+++ b/Engine/Plugins/Experimental/DataPrepEditor/Source/DataPrepLibraries/Private/Fetchers/DataprepFloatFetcherLibrary.cpp
@@ -5,9 +5,7 @@
 #include "Components/StaticMeshComponent.h"
 #include "Engine/StaticMesh.h"
 #include "GameFramework/Actor.h"
-#include "MeshAttributeArray.h"
-#include "MeshAttributes.h"
-#include "MeshDescription.h"
+#include "StaticMeshAttributes.h"
 #include "Misc/Optional.h"
 #include "UObject/Object.h"
 
@@ -24,7 +22,8 @@ float UDataprepFloatBoundingVolumeFetcher::Fetch_Implementation(const UObject* O
 					FBox Box(ForceInit);
 					if ( const FMeshDescription* MeshDescription = StaticMesh->GetMeshDescription( 0 ) )
 					{
-						TVertexAttributesConstRef<FVector> VertexPositions = MeshDescription->VertexAttributes().GetAttributesRef<FVector>( MeshAttribute::Vertex::Position );
+						FStaticMeshConstAttributes Attributes(*MeshDescription);
+						TVertexAttributesConstRef<FVector> VertexPositions = Attributes.GetVertexPositions();
 						for ( const FVertexID VertexID : MeshDescription->Vertices().GetElementIDs() )
 						{
 							if ( !MeshDescription->IsVertexOrphaned( VertexID ) )
diff --git a/Engine/Plugins/Experimental/MeshModelingToolset/MeshModelingToolset.uplugin b/Engine/Plugins/Experimental/MeshModelingToolset/MeshModelingToolset.uplugin
new file mode 100644
--- /dev/null
+++ b/Engine/Plugins/Experimental/MeshModelingToolset/MeshModelingToolset.uplugin
@@ -0,0 +1,58 @@
+{
+	"FileVersion": 3,
+	"Version": 1,
+	"VersionName": "0.1",
+	"FriendlyName": "Mesh Modeling Toolset",
+	"Description": "A set of modules implementing 3D mesh creation and editing based on the Interactive Tools Framework",
+	"Category": "Other",
+	"CreatedBy": "Epic Games, Inc.",
+	"CreatedByURL": "http://epicgames.com",
+	"DocsURL": "",
+	"MarketplaceURL": "",
+	"SupportURL": "",
+	"CanContainContent": true,
+	"IsBetaVersion": true,
+	"Installed": false,
+	"Modules": [
+		{
+			"Name": "ModelingComponents",
+			"Type": "Runtime",
+			"LoadingPhase": "Default",
+			"WhitelistPlatforms": [ "Win64" ]
+		},
+		{
+			"Name": "ModelingOperators",
+			"Type": "Runtime",
+			"LoadingPhase": "Default",
+			"WhitelistPlatforms": [ "Win64" ]
+		},
+		{
+			"Name": "ModelingOperatorsEditorOnly",
+			"Type": "Editor",
+			"LoadingPhase": "Default",
+			"WhitelistPlatforms": [ "Win64" ]
+		},
+		{
+			"Name": "MeshModelingTools",
+			"Type": "Runtime",
+			"LoadingPhase": "Default",
+			"WhitelistPlatforms": [ "Win64" ]
+		},
+		{
+			"Name": "MeshModelingToolsEditorOnly",
+			"Type": "Editor",
+			"LoadingPhase": "Default",
+			"WhitelistPlatforms": [ "Win64" ]
+		},
+	],
+	"Plugins": [
+		{
+			"Name": "GeometryProcessing",
+			"Enabled": true
+		},
+		{
+			"Name": "ProxyLODPlugin",
+			"Enabled" : true
+		}
+	]
+}
diff --git a/Engine/Plugins/Experimental/MeshModelingToolset/Source/MeshModelingTools/MeshModelingTools.Build.cs b/Engine/Plugins/Experimental/MeshModelingToolset/Source/MeshModelingTools/MeshModelingTools.Build.cs
new file mode 100644
--- /dev/null
+++ b/Engine/Plugins/Experimental/MeshModelingToolset/Source/MeshModelingTools/MeshModelingTools.Build.cs
@@ -0,0 +1,88 @@
+// Copyright 1998-2019 Epic Games, Inc. All Rights Reserved.
+
+using UnrealBuildTool;
+
+public class MeshModelingTools : ModuleRules
+{
+	public MeshModelingTools(ReadOnlyTargetRules Target) : base(Target)
+	{
+		PCHUsage = ModuleRules.PCHUsageMode.UseExplicitOrSharedPCHs;
+		
+		PublicIncludePaths.AddRange(
+			new string[] {
+				// ... add public include paths required here ...
+			}
+			);
+				
+		
+		PrivateIncludePaths.AddRange(
+			new string[] {
+				// ... add other private include paths required here ...
+			}
+			);
+			
+		
+		PublicDependencyModuleNames.AddRange(
+			new string[]
+			{
+				// These dependencies were commented out when we split out MeshModelingToolsEditorOnly
+				// Many of them are editor only and should not be depended on here.  If you re-add any
+				// dependent modules, please confirm by "launching" or packaging that you are not introducing
+				// editor dependencies
+				//
+				//"MeshDescription",
+				//"ProxyLODMeshReduction", // currently required to be public due to IVoxelBasedCSG API
+
+				"Core",
+				"Eigen",
+                "InteractiveToolsFramework",
+				"GeometricObjects",
+				"DynamicMesh",
+				"MeshConversion",
+                "MeshSolverUtilities",
+				"ModelingComponents",
+				"ModelingOperators",
+
+				// ... add other public dependencies that you statically link with here ...
+			}
+            );
+			
+		
+		PrivateDependencyModuleNames.AddRange(
+			new string[]
+			{
+				// These dependencies were commented out when we split out MeshModelingToolsEditorOnly
+				// Many of them are editor only and should not be depended on here.  If you re-add any
+				// dependent modules, please confirm by "launching" or packaging that you are not introducing
+				// editor dependencies
+				//
+				//"RenderCore",
+				//"RHI",
+				//"MeshUtilities",    // temp for saving mesh asset
+				//"UnrealEd",
+				//"MeshBuilder",
+                //"MeshDescriptionOperations",
+				//"MeshUtilitiesCommon",
+				//"MeshReductionInterface", // for UE4 standard simplification
+                //"ProxyLODMeshReduction", // for mesh merging voxel-based csg
+				//"Slate",
+				//"SlateCore",
+
+				"CoreUObject",
+				"Engine",
+				"ModelingOperators",
+				"InputCore",
+
+				// ... add private dependencies that you statically link with here ...
+			}
+			);
+		
+		
+		DynamicallyLoadedModuleNames.AddRange(
+			new string[]
+			{
+				// ... add any modules that your module loads dynamically here ...
+			}
+			);
+	}
+}
diff --git a/Engine/Plugins/Experimental/MeshModelingToolset/Source/MeshModelingTools/Private/ConvertToPolygonsTool.cpp b/Engine/Plugins/Experimental/MeshModelingToolset/Source/MeshModelingTools/Private/ConvertToPolygonsTool.cpp
new file mode 100644
--- /dev/null
+++ b/Engine/Plugins/Experimental/MeshModelingToolset/Source/MeshModelingTools/Private/ConvertToPolygonsTool.cpp
@@ -0,0 +1,209 @@
+// Copyright 1998-2019 Epic Games, Inc. All Rights Reserved.
+
+#include "ConvertToPolygonsTool.h"
+#include "InteractiveToolManager.h"
+#include "ToolBuilderUtil.h"
+
+#include "DynamicMesh3.h"
+#include "MeshNormals.h"
+#include "DynamicMeshEditor.h"
+#include "MeshRegionBoundaryLoops.h"
+#include "MeshDescriptionToDynamicMesh.h"
+#include "DynamicMeshToMeshDescription.h"
+
+
+#include "SceneManagement.h" // for FPrimitiveDrawInterface
+
+#define LOCTEXT_NAMESPACE "UConvertToPolygonsTool"
+
+
+/*
+ * ToolBuilder
+ */
+
+
+bool UConvertToPolygonsToolBuilder::CanBuildTool(const FToolBuilderState& SceneState) const
+{
+	return ToolBuilderUtil::CountComponents(SceneState, CanMakeComponentTarget) == 1;
+}
+
+UInteractiveTool* UConvertToPolygonsToolBuilder::BuildTool(const FToolBuilderState& SceneState) const
+{
+	UConvertToPolygonsTool* NewTool = NewObject<UConvertToPolygonsTool>(SceneState.ToolManager);
+
+	UActorComponent* ActorComponent = ToolBuilderUtil::FindFirstComponent(SceneState, CanMakeComponentTarget);
+	UPrimitiveComponent* MeshComponent = Cast<UPrimitiveComponent>(ActorComponent);
+	check(MeshComponent != nullptr);
+	NewTool->SetSelection(MakeComponentTarget(MeshComponent));
+
+	return NewTool;
+}
+
+
+
+/*
+ * Tool
+ */
+
+
+void UConvertToPolygonsTool::Setup()
+{
+	UInteractiveTool::Setup();
+
+	FMeshDescription* MeshDescription = nullptr;
+	MeshDescription = ComponentTarget->GetMesh();
+
+	FMeshDescriptionToDynamicMesh Converter;
+	Converter.Convert(MeshDescription, SearchMesh);
+
+	if (SearchMesh.HasAttributes())
+	{
+		InitialNormals.Copy( *SearchMesh.Attributes()->PrimaryNormals() );
+	}
+
+	Settings = NewObject<UConvertToPolygonsToolProperties>(this);
+	AddToolPropertySource(Settings);
+
+	// create preview mesh object
+	PreviewMesh = NewObject<UPreviewMesh>(this, TEXT("PreviewMesh"));
+	PreviewMesh->CreateInWorld(ComponentTarget->GetOwnerActor()->GetWorld(), FTransform::Identity);
+	PreviewMesh->SetVisible(false);
+	PreviewMesh->SetTransform(ComponentTarget->GetWorldTransform());
+
+	if (ComponentTarget->GetMaterial(0) != nullptr)
+	{
+		PreviewMesh->SetMaterial(ComponentTarget->GetMaterial(0));
+	}
+
+	UpdatePolygons();
+}
+
+
+void UConvertToPolygonsTool::Shutdown(EToolShutdownType ShutdownType)
+{
+	PreviewMesh->SetVisible(false);
+	PreviewMesh->Disconnect();
+	PreviewMesh = nullptr;
+
+	ComponentTarget->SetOwnerVisibility(true);
+
+	if (ShutdownType == EToolShutdownType::Accept)
+	{
+		GetToolManager()->BeginUndoTransaction(LOCTEXT("DeformMeshToolTransactionName", "Convert to Polygons"));
+		ComponentTarget->CommitMesh([=](FMeshDescription* MeshDescription)
+		{
+			ConvertToPolygons(MeshDescription);
+		});
+		GetToolManager()->EndUndoTransaction();
+	}
+}
+
+
+void UConvertToPolygonsTool::OnPropertyModified(UObject* PropertySet, UProperty* Property)
+{
+	bPolygonsValid = false;
+	GetToolManager()->PostInvalidation();
+}
+
+
+void UConvertToPolygonsTool::Render(IToolsContextRenderAPI* RenderAPI)
+{
+	FColor LineColor(255, 0, 0);
+
+	if (bPolygonsValid == false)
+	{
+		UpdatePolygons();
+	}
+
+	if (true)
+	{
+		FPrimitiveDrawInterface* PDI = RenderAPI->GetPrimitiveDrawInterface();
+		FTransform Transform = ComponentTarget->GetWorldTransform(); //Actor->GetTransform();
+
+		TArray<int>& Edges = Polygons.PolygonEdges;
+		for (int eid : Edges)
+		{
+			FVector3d A, B;
+			Polygons.Mesh->GetEdgeV(eid, A, B);
+			PDI->DrawLine(Transform.TransformPosition(A), Transform.TransformPosition(B),
+				LineColor, 0, 2.0, 1.0f, true);
+		}
+	}
+}
+
+
+bool UConvertToPolygonsTool::HasAccept() const
+{
+	return true;
+}
+
+bool UConvertToPolygonsTool::CanAccept() const
+{
+	return true;
+}
+
+
+void UConvertToPolygonsTool::UpdatePolygons()
+{
+	Polygons = FFindPolygonsAlgorithm(&SearchMesh);
+	double DotTolerance = 1.0 - FMathd::Cos(Settings->AngleTolerance * FMathd::DegToRad);
+	Polygons.FindPolygons(DotTolerance);
+	Polygons.FindPolygonEdges();
+
+	GetToolManager()->DisplayMessage(
+		FText::Format(LOCTEXT("UpdatePolygonsMessage", "ConvertToPolygons - found {0} polys in {1} triangles"), 
+			FText::AsNumber(Polygons.FoundPolygons.Num()), FText::AsNumber(SearchMesh.TriangleCount()) ), EToolMessageLevel::Internal);
+
+	if (Settings->bCalculateNormals)
+	{
+		if (SearchMesh.HasAttributes() == false)
+		{
+			SearchMesh.EnableAttributes();
+		}
+
+		FDynamicMeshNormalOverlay* NormalOverlay = SearchMesh.Attributes()->PrimaryNormals();
+		NormalOverlay->ClearElements();
+
+		FDynamicMeshEditor Editor(&SearchMesh);
+		for (const TArray<int>& Polygon : Polygons.FoundPolygons)
+		{
+			FVector3f Normal = (FVector3f)SearchMesh.GetTriNormal(Polygon[0]);
+			Editor.SetTriangleNormals(Polygon, Normal);
+		}
+
+		FMeshNormals Normals(&SearchMesh);
+		Normals.RecomputeOverlayNormals(SearchMesh.Attributes()->PrimaryNormals());
+		Normals.CopyToOverlay(NormalOverlay, false);
+
+		PreviewMesh->UpdatePreview(&SearchMesh);
+		PreviewMesh->SetVisible(true);
+		ComponentTarget->SetOwnerVisibility(false);
+	}
+	else
+	{
+		PreviewMesh->SetVisible(false);
+		ComponentTarget->SetOwnerVisibility(true);
+	}
+
+	bPolygonsValid = true;
+}
+
+
+
+void UConvertToPolygonsTool::ConvertToPolygons(FMeshDescription* MeshIn)
+{
+	// restore input normals
+	if (Settings->bCalculateNormals == false)
+	{
+		SearchMesh.Attributes()->PrimaryNormals()->Copy(InitialNormals);
+	}
+
+	FDynamicMeshToMeshDescription Converter;
+	Converter.ConversionOptions.bSetPolyGroups = true;
+	Converter.Convert(&SearchMesh, *MeshIn);
+}
+
+
+
+
+#undef LOCTEXT_NAMESPACE
diff --git a/Engine/Plugins/Experimental/MeshModelingToolset/Source/MeshModelingTools/Private/DisplaceMeshTool.cpp b/Engine/Plugins/Experimental/MeshModelingToolset/Source/MeshModelingTools/Private/DisplaceMeshTool.cpp
new file mode 100644
--- /dev/null
+++ b/Engine/Plugins/Experimental/MeshModelingToolset/Source/MeshModelingTools/Private/DisplaceMeshTool.cpp
@@ -0,0 +1,434 @@
+// Copyright 1998-2019 Epic Games, Inc. All Rights Reserved.
+
+#include "DisplaceMeshTool.h"
+#include "InteractiveToolManager.h"
+#include "ToolBuilderUtil.h"
+
+#include "DynamicMesh3.h"
+
+#include "SimpleDynamicMeshComponent.h"
+
+#include "MeshNormals.h"
+
+#include "SceneManagement.h" // for FPrimitiveDrawInterface
+#include "Async/ParallelFor.h"
+
+
+#include "ProfilingDebugging/ScopedTimers.h"
+
+#define LOCTEXT_NAMESPACE "UDisplaceMeshTool"
+
+/*
+ * ToolBuilder
+ */
+bool UDisplaceMeshToolBuilder::CanBuildTool(const FToolBuilderState& SceneState) const
+{
+	return ToolBuilderUtil::CountComponents(SceneState, CanMakeComponentTarget) == 1;
+}
+
+UInteractiveTool* UDisplaceMeshToolBuilder::BuildTool(const FToolBuilderState& SceneState) const
+{
+	UDisplaceMeshTool* NewTool = NewObject<UDisplaceMeshTool>(SceneState.ToolManager);
+
+	UActorComponent* ActorComponent = ToolBuilderUtil::FindFirstComponent(SceneState, CanMakeComponentTarget);
+	auto* MeshComponent = Cast<UPrimitiveComponent>(ActorComponent);
+	check(MeshComponent != nullptr);
+
+	NewTool->SetSelection(MakeComponentTarget(MeshComponent));
+
+	return NewTool;
+}
+
+/*
+ * Tool
+ */
+UDisplaceMeshTool::UDisplaceMeshTool()
+{
+	DisplacementType = EDisplaceMeshToolDisplaceType::DisplacementMap;
+	DisplaceIntensity = 10.0f;
+	RandomSeed = 31337;
+	Subdivisions = 4;
+}
+
+void UDisplaceMeshTool::Setup()
+{
+	UInteractiveTool::Setup();
+
+	// create dynamic mesh component to use for live preview
+	DynamicMeshComponent = NewObject<USimpleDynamicMeshComponent>(ComponentTarget->GetOwnerActor(), "DynamicMesh");
+	DynamicMeshComponent->SetupAttachment(ComponentTarget->GetOwnerActor()->GetRootComponent());
+	DynamicMeshComponent->RegisterComponent();
+	DynamicMeshComponent->SetWorldTransform(ComponentTarget->GetWorldTransform());
+
+	// copy material if there is one
+	auto Material = ComponentTarget->GetMaterial(0);
+	if (Material != nullptr)
+	{
+		DynamicMeshComponent->SetMaterial(0, Material);
+	}
+
+	DynamicMeshComponent->InitializeMesh(ComponentTarget->GetMesh());
+	OriginalMesh.Copy(*DynamicMeshComponent->GetMesh());
+
+	UpdateMap(true);
+
+	// hide input StaticMeshComponent
+	ComponentTarget->SetOwnerVisibility(false);
+
+	// initialize our properties
+	ToolPropertyObjects.Add(this);
+
+	bResultValid = false;
+}
+
+void UDisplaceMeshTool::Shutdown(EToolShutdownType ShutdownType)
+{
+	if (DynamicMeshComponent != nullptr)
+	{
+		//DynamicMeshComponent->OnMeshChanged.Remove(OnDynamicMeshComponentChangedHandle);
+
+		ComponentTarget->SetOwnerVisibility(true);
+
+		if (ShutdownType == EToolShutdownType::Accept)
+		{
+			// this block bakes the modified DynamicMeshComponent back into the StaticMeshComponent inside an undo transaction
+			GetToolManager()->BeginUndoTransaction(LOCTEXT("DisplaceMeshToolTransactionName", "Displace Mesh"));
+			ComponentTarget->CommitMesh([=](FMeshDescription* MeshDescription)
+			{
+				DynamicMeshComponent->Bake(MeshDescription, Subdivisions > 0);
+			});
+			GetToolManager()->EndUndoTransaction();
+		}
+
+		DynamicMeshComponent->UnregisterComponent();
+		DynamicMeshComponent->DestroyComponent();
+		DynamicMeshComponent = nullptr;
+	}
+}
+
+void UDisplaceMeshTool::Render(IToolsContextRenderAPI* RenderAPI)
+{
+	UpdateResult();
+}
+
+#if WITH_EDITOR
+void UDisplaceMeshTool::PostEditChangeProperty(FPropertyChangedEvent& PropertyChangedEvent)
+{
+	UProperty* PropertyThatChanged = PropertyChangedEvent.Property;
+
+	if (PropertyThatChanged)
+	{
+		//if (PropertyThatChanged->GetFName() == GET_MEMBER_NAME_CHECKED(UDisplaceMeshTool, DisplaceIterations))
+		//{
+		//}
+	}
+
+	bResultValid = false;
+}
+#endif
+
+void UDisplaceMeshTool::UpdateResult()
+{
+	if (bResultValid) 
+	{
+		return;
+	}
+
+	UpdateSubdividedMesh();
+
+	// compute Displaceed positions in PositionBuffer
+	switch (DisplacementType)
+	{
+		default:
+		case EDisplaceMeshToolDisplaceType::Constant:
+			ComputeDisplacement_Constant();
+			break;
+
+		case EDisplaceMeshToolDisplaceType::RandomNoise:
+			ComputeDisplacement_RandomNoise();
+			break;
+
+		case EDisplaceMeshToolDisplaceType::DisplacementMap:
+			ComputeDisplacement_Map();
+			break;
+	}
+
+	// update preview vertex positions
+	FDynamicMesh3* TargetMesh = DynamicMeshComponent->GetMesh();
+	for (int vid : SubdividedMesh.VertexIndicesItr())
+	{
+		TargetMesh->SetVertex(vid, DisplacedBuffer[vid]);
+	}
+
+	// recalculate normals
+	if (TargetMesh->HasAttributes())
+	{
+		FMeshNormals Normals(TargetMesh);
+		FDynamicMeshNormalOverlay* NormalOverlay = TargetMesh->Attributes()->PrimaryNormals();
+		Normals.RecomputeOverlayNormals(NormalOverlay);
+		Normals.CopyToOverlay(NormalOverlay);
+	}
+	else
+	{
+		FMeshNormals::QuickComputeVertexNormals(*TargetMesh);
+	}
+
+	DynamicMeshComponent->NotifyMeshUpdated();
+	//DynamicMeshComponent->FastNotifyPositionsUpdated(true);
+
+	GetToolManager()->PostInvalidation();
+
+	bResultValid = true;
+}
+
+void UDisplaceMeshTool::ComputeDisplacement_Constant()
+{
+	double Intensity = DisplaceIntensity;
+
+	for (int vid : SubdividedMesh.VertexIndicesItr())
+	{
+		FVector3d Position = PositionBuffer[vid];
+		FVector3d Normal = NormalsBuffer[vid];
+		FVector3d Displacement = Intensity * Normal;
+		DisplacedBuffer[vid] = Position + Displacement;
+	}
+}
+
+void UDisplaceMeshTool::ComputeDisplacement_RandomNoise()
+{
+	double Intensity = DisplaceIntensity;
+
+	FMath::SRandInit(RandomSeed);
+
+	for (int vid : SubdividedMesh.VertexIndicesItr())
+	{
+		FVector3d Position = PositionBuffer[vid];
+		FVector3d Normal = NormalsBuffer[vid];
+
+		double RandVal = 2.0 * (FMath::SRand() - 0.5);
+		
+		FVector3d Displacement = RandVal * Intensity * Normal;
+
+		DisplacedBuffer[vid] = Position + Displacement;
+	}
+}
+
+void UDisplaceMeshTool::ComputeDisplacement_Map()
+{
+	UpdateMap();
+
+	const FDynamicMeshUVOverlay* UVOverlay = SubdividedMesh.Attributes()->GetUVLayer(0);
+
+	double Intensity = DisplaceIntensity;
+	FMath::SRandInit(RandomSeed);
+
+	for (int tid : SubdividedMesh.TriangleIndicesItr())
+	{
+		FIndex3i Tri = SubdividedMesh.GetTriangle(tid);
+		FIndex3i UVTri = UVOverlay->GetTriangle(tid);
+
+		for (int j = 0; j < 3; ++j)
+		{
+			int vid = Tri[j];
+
+			FVector3d Position = PositionBuffer[vid];
+			FVector3d Normal = NormalsBuffer[vid];
+
+			FVector2f UV = UVOverlay->GetElement(UVTri[j]);
+			double Offset = DisplaceField.BilinearSampleClamped(UV);
+
+			FVector3d Displacement = Offset * Intensity * Normal;
+
+			DisplacedBuffer[vid] = Position + Displacement;
+		}
+	}
+
+}
+
+namespace {
+	class FTextureAccess
+	{
+	public:
+		FTextureAccess(UTexture2D* DisplacementMap)
+			:DisplacementMap(DisplacementMap)
+		{
+			check(DisplacementMap);
+			OldCompressionSettings = DisplacementMap->CompressionSettings;
+			bOldSRGB = DisplacementMap->SRGB; 
+#if WITH_EDITOR
+			OldMipGenSettings = DisplacementMap->MipGenSettings;
+#endif
+
+
+			DisplacementMap->CompressionSettings = TextureCompressionSettings::TC_VectorDisplacementmap;
+			DisplacementMap->SRGB = false;
+#if WITH_EDITOR
+			DisplacementMap->MipGenSettings = TextureMipGenSettings::TMGS_NoMipmaps;
+#endif
+			DisplacementMap->UpdateResource();
+
+			FormattedImageData = static_cast<const FColor*>(DisplacementMap->PlatformData->Mips[0].BulkData.LockReadOnly());
+		}
+		FTextureAccess(const FTextureAccess&) = delete;
+		FTextureAccess(FTextureAccess&&) = delete;
+		void operator=(const FTextureAccess&) = delete;
+		void operator=(FTextureAccess&&) = delete;
+
+		~FTextureAccess()
+		{
+			DisplacementMap->PlatformData->Mips[0].BulkData.Unlock();
+
+			DisplacementMap->CompressionSettings = OldCompressionSettings;
+			DisplacementMap->SRGB = bOldSRGB;
+#if WITH_EDITOR
+			DisplacementMap->MipGenSettings = OldMipGenSettings;
+#endif
+
+			DisplacementMap->UpdateResource();
+		}
+
+		bool HasData() const
+		{
+			return FormattedImageData != nullptr;
+		}
+		const FColor* GetData() const
+		{
+			return FormattedImageData;
+		}
+
+	private:
+		UTexture2D* DisplacementMap{nullptr};
+		TextureCompressionSettings OldCompressionSettings{};
+		TextureMipGenSettings OldMipGenSettings{};
+		bool bOldSRGB{false};
+		const FColor* FormattedImageData{nullptr};
+	};
+}
+
+void UDisplaceMeshTool::UpdateMap(bool bForceUpdate)
+{
+	if (CachedMapSource == DisplacementMap && bForceUpdate == false)
+	{
+		return;
+	}
+	CachedMapSource = DisplacementMap;
+
+	if (DisplacementMap == nullptr || DisplacementMap->PlatformData == nullptr || DisplacementMap->PlatformData->Mips.Num() < 1)
+	{
+		DisplaceField = FSampledScalarField2f();
+		return;
+	}
+
+	FTextureAccess TextureAccess( DisplacementMap );
+	if (!TextureAccess.HasData())
+	{
+		DisplaceField = FSampledScalarField2f();
+	}
+	else
+	{
+		int TextureWidth = DisplacementMap->GetSizeX();
+		int TextureHeight = DisplacementMap->GetSizeY();
+		DisplaceField.Resize(TextureWidth, TextureHeight, 0.0f);
+		DisplaceField.SetCellSize(1.0f / (float)TextureWidth);
+
+		const FColor* FormattedData = TextureAccess.GetData();
+		for (int y = 0; y < TextureHeight; ++y)
+		{
+			for (int x = 0; x < TextureWidth; ++x)
+			{
+				FColor PixelColor = FormattedData[y*TextureWidth + x];
+				float Value = PixelColor.R / 255.0;
+				DisplaceField.GridValues[y*TextureWidth + x] = Value;
+			}
+		}
+	}
+}
+
+void UDisplaceMeshTool::UpdateSubdividedMesh()
+{
+	if (Subdivisions == CachedSubdivisionsCount)
+	{
+		return;
+	}
+
+	SubdividedMesh.Copy(OriginalMesh);
+	int Rounds = Subdivisions;
+
+	// calculate subdivisions (todo: move to elsewhere)
+	for (int ri = 0; ri < Rounds; ri++)
+	{
+		TArray<int> EdgesToProcess;
+		for (int tid : SubdividedMesh.EdgeIndicesItr())
+		{
+			EdgesToProcess.Add(tid);
+		}
+		int MaxTriangleID = SubdividedMesh.MaxTriangleID();
+
+		TArray<int> TriSplitEdges;
+		TriSplitEdges.Init(-1, SubdividedMesh.MaxTriangleID());
+
+		for (int eid : EdgesToProcess)
+		{
+			FIndex2i EdgeTris = SubdividedMesh.GetEdgeT(eid);
+
+			FDynamicMesh3::FEdgeSplitInfo SplitInfo;
+			EMeshResult result = SubdividedMesh.SplitEdge(eid, SplitInfo);
+			if (result == EMeshResult::Ok)
+			{
+				if (EdgeTris.A < MaxTriangleID && TriSplitEdges[EdgeTris.A] == -1)
+				{
+					TriSplitEdges[EdgeTris.A] = SplitInfo.NewEdges.B;
+				}
+				if (EdgeTris.B != FDynamicMesh3::InvalidID)
+				{
+					if (EdgeTris.B < MaxTriangleID && TriSplitEdges[EdgeTris.B] == -1)
+					{
+						TriSplitEdges[EdgeTris.B] = SplitInfo.NewEdges.C;
+					}
+				}
+			}
+		}
+
+		for (int eid : TriSplitEdges)
+		{
+			if (eid != -1)
+			{
+				FDynamicMesh3::FEdgeFlipInfo FlipInfo;
+				SubdividedMesh.FlipEdge(eid, FlipInfo);
+			}
+		}
+	}
+
+	CachedSubdivisionsCount = Subdivisions;
+
+	// update dynamic mesh component
+	DynamicMeshComponent->GetMesh()->Copy(SubdividedMesh);
+	DynamicMeshComponent->NotifyMeshUpdated();
+
+	NormalsBuffer = FMeshNormals(&SubdividedMesh);
+	NormalsBuffer.ComputeVertexNormals();
+
+	// cache initial positions
+	PositionBuffer.SetNum(SubdividedMesh.MaxVertexID());
+	for (int vid : SubdividedMesh.VertexIndicesItr())
+	{
+		PositionBuffer[vid] = SubdividedMesh.GetVertex(vid);
+	}
+	DisplacedBuffer.SetNum(SubdividedMesh.MaxVertexID());
+}
+
+
+bool UDisplaceMeshTool::HasAccept() const
+{
+	return true;
+}
+
+bool UDisplaceMeshTool::CanAccept() const
+{
+	return true;
+}
+
+
+
+
+#undef LOCTEXT_NAMESPACE
