Fixed a crash in the speed tree importer.

Added a import context the speed tree importer to avoid reprocessing/importing multiple time the same texture or material.
Added a function to Package Tools since the patern that is causing the crash is a recuring one.

diff --git a/Engine/Plugins/Editor/SpeedTreeImporter/Source/SpeedTreeImporter/Private/SpeedTreeImportFactory.cpp b/Engine/Plugins/Editor/SpeedTreeImporter/Source/SpeedTreeImporter/Private/SpeedTreeImportFactory.cpp
--- a/Engine/Plugins/Editor/SpeedTreeImporter/Source/SpeedTreeImporter/Private/SpeedTreeImportFactory.cpp
+++ b/Engine/Plugins/Editor/SpeedTreeImporter/Source/SpeedTreeImporter/Private/SpeedTreeImportFactory.cpp
@@ -1,40 +1,41 @@
 // Copyright Epic Games, Inc. All Rights Reserved.
 
 #include "SpeedTreeImportFactory.h"
+
+#include "Editor.h"
+#include "EditorFramework/AssetImportData.h"
+#include "EditorReimportHandler.h"
+#include "EditorStyleSet.h"
+#include "Engine/StaticMesh.h"
+#include "Engine/Texture.h"
+#include "Engine/Texture2D.h"
+#include "Factories/MaterialFactoryNew.h"
+#include "Factories/TextureFactory.h"
+#include "Framework/Application/SlateApplication.h"
+#include "Input/Reply.h"
+#include "Materials/Material.h"
+#include "Materials/MaterialExpressionMaterialFunctionCall.h"
+#include "Materials/MaterialFunction.h"
+#include "Materials/MaterialInterface.h"
 #include "Misc/FileHelper.h"
+#include "Misc/PackageName.h"
 #include "Misc/Paths.h"
 #include "Modules/ModuleManager.h"
-#include "Misc/PackageName.h"
-#include "Input/Reply.h"
+#include "SlateOptMacros.h"
+#include "Styling/SlateTypes.h"
+#include "UObject/GCObject.h"
 #include "Widgets/DeclarativeSyntaxSupport.h"
-#include "Widgets/SWidget.h"
-#include "Widgets/SCompoundWidget.h"
+#include "Widgets/Input/SButton.h"
+#include "Widgets/Input/SCheckBox.h"
+#include "Widgets/Input/SEditableTextBox.h"
+#include "Widgets/Layout/SBorder.h"
+#include "Widgets/Layout/SBox.h"
+#include "Widgets/Layout/SUniformGridPanel.h"
 #include "Widgets/SBoxPanel.h"
-#include "Styling/SlateTypes.h"
+#include "Widgets/SCompoundWidget.h"
+#include "Widgets/SWidget.h"
 #include "Widgets/SWindow.h"
-#include "SlateOptMacros.h"
-#include "Framework/Application/SlateApplication.h"
-#include "Widgets/Layout/SBox.h"
-#include "Widgets/Layout/SBorder.h"
 #include "Widgets/Text/STextBlock.h"
-#include "Widgets/Layout/SUniformGridPanel.h"
-#include "Widgets/Input/SEditableTextBox.h"
-#include "Widgets/Input/SButton.h"
-#include "Widgets/Input/SCheckBox.h"
-#include "EditorStyleSet.h"
-#include "EditorReimportHandler.h"
-#include "Materials/MaterialInterface.h"
-#include "Materials/MaterialExpressionMaterialFunctionCall.h"
-#include "Materials/MaterialFunction.h"
-#include "Materials/Material.h"
-#include "Factories/MaterialFactoryNew.h"
-#include "Engine/Texture.h"
-#include "Factories/TextureFactory.h"
-#include "EditorFramework/AssetImportData.h"
-#include "Engine/Texture2D.h"
-#include "Engine/StaticMesh.h"
-#include "Editor.h"
 
-
 #include "Materials/MaterialExpressionClamp.h"
 #include "Materials/MaterialExpressionComponentMask.h"
@@ -244,6 +245,26 @@
 END_SLATE_FUNCTION_BUILD_OPTIMIZATION
 
 
+/**
+ * A small context to avoid importing multiple time the same file
+ */
+struct FSpeedTreeImportContext : public FGCObject
+{
+	TMap<FString, UMaterialInterface*> ImportedMaterials;
+	TMap<FString, UTexture*> ImportedTextures;
+
+	virtual void AddReferencedObjects(FReferenceCollector& Collector) override
+	{
+		TArray<UMaterialInterface*> Materials;
+		ImportedMaterials.GenerateValueArray(Materials);
+		Collector.AddReferencedObjects(Materials);
+
+		TArray<UTexture*> Textures;
+		ImportedTextures.GenerateValueArray(Textures);
+		Collector.AddReferencedObjects(ImportedTextures);
+	}
+};
+
 //////////////////////////////////////////////////////////////////////////
 
 USpeedTreeImportFactory::USpeedTreeImportFactory(const FObjectInitializer& ObjectInitializer)
@@ -330,7 +351,7 @@
 	return UStaticMesh::StaticClass();
 }
 
-UTexture* CreateSpeedTreeMaterialTexture(UObject* Parent, FString Filename, bool bNormalMap, bool bMasks, TSet<UPackage*>& LoadedPackages)
+UTexture* CreateSpeedTreeMaterialTexture(UObject* Parent,const FString& Filename, bool bNormalMap, bool bMasks, TSet<UPackage*>& LoadedPackages, FSpeedTreeImportContext& ImportContext)
 {
 	UTexture* UnrealTexture = NULL;
 
@@ -339,14 +360,22 @@
 		return UnrealTexture;
 	}
 
+	if (UTexture** Texture = ImportContext.ImportedTextures.Find(Filename))
+	{
+		// The texture was already processed
+		return *Texture;
+	}
+
 	FString Extension = FPaths::GetExtension(Filename).ToLower();
 	FString TextureName = FPaths::GetBaseFilename(Filename) + TEXT("_Tex");
 	TextureName = ObjectTools::SanitizeObjectName(TextureName);
 
 	// set where to place the textures
 	FString NewPackageName = FPackageName::GetLongPackagePath(Parent->GetOutermost()->GetName()) + TEXT("/") + TextureName;
-	NewPackageName = UPackageTools::SanitizePackageName(NewPackageName);
-	UPackage* Package = CreatePackage(NULL, *NewPackageName);
+	UPackage* Package = UPackageTools::FindOrCreatePackageForAssetType(FName(*NewPackageName), UTexture2D::StaticClass());
+	check(Package);
+	NewPackageName = Package->GetFullName();
+	TextureName = FPaths::GetBaseFilename(NewPackageName, true);
 
 #if 0
 	// find existing texture
@@ -367,11 +396,11 @@
 #endif
 
 	// try opening from absolute path
-	Filename = FPaths::GetPath(UFactory::GetCurrentFilename()) + "/" + Filename;
+	FString FilenamePath = FPaths::GetPath(UFactory::GetCurrentFilename()) + "/" + Filename;
 	TArray<uint8> TextureData;
-	if (!(FFileHelper::LoadFileToArray(TextureData, *Filename) && TextureData.Num() > 0))
+	if (!(FFileHelper::LoadFileToArray(TextureData, *FilenamePath) && TextureData.Num() > 0))
 	{
-		UE_LOG(LogSpeedTreeImport, Warning, TEXT("Unable to find Texture file %s"), *Filename);
+		UE_LOG(LogSpeedTreeImport, Warning, TEXT("Unable to find Texture file %s"), *FilenamePath);
 	}
 	else
 	{
@@ -400,7 +429,7 @@
 				UnrealTexture->SRGB = false;
 				UnrealTexture->CompressionSettings = TC_Masks;
 			}
-			UnrealTexture->AssetImportData->Update(Filename);
+			UnrealTexture->AssetImportData->Update(FilenamePath);
 
 			// Notify the asset registry
 			FAssetRegistryModule::AssetCreated(UnrealTexture);
@@ -413,6 +442,9 @@
 		TextureFact->RemoveFromRoot();
 	}
 
+	// Keep track of the created textures
+	ImportContext.ImportedTextures.Add(Filename, UnrealTexture);
+
 	return UnrealTexture;
 }
 
@@ -530,7 +562,7 @@
 	}
 }
 
-UMaterialInterface* CreateSpeedTreeMaterial7(UObject* Parent, FString MaterialFullName, const SpeedTree::SRenderState* RenderState, TSharedPtr<SSpeedTreeImportOptions> Options, ESpeedTreeWindType WindType, int32 NumBillboards, TSet<UPackage*>& LoadedPackages)
+UMaterialInterface* CreateSpeedTreeMaterial7(UObject* Parent, FString MaterialFullName, const SpeedTree::SRenderState* RenderState, TSharedPtr<SSpeedTreeImportOptions> Options, ESpeedTreeWindType WindType, int32 NumBillboards, TSet<UPackage*>& LoadedPackages, FSpeedTreeImportContext& ImportContext)
 {
 	// Make sure we have a parent
 	if (!Options->SpeedTreeImportData->MakeMaterialsCheck || !ensure(Parent))
@@ -538,32 +570,43 @@
 		return UMaterial::GetDefaultMaterial(MD_Surface);
 	}
 	
+	if (UMaterialInterface** Material = ImportContext.ImportedMaterials.Find(MaterialFullName))
+	{
+		// The material was already imported
+		return *Material;
+	}
+
 	// set where to place the materials
 	FString FixedMaterialName = MaterialFullName + TEXT("_Mat");
 	FString NewPackageName = FPackageName::GetLongPackagePath(Parent->GetOutermost()->GetName()) + TEXT("/") + FixedMaterialName;
-	NewPackageName = UPackageTools::SanitizePackageName(NewPackageName);
-	UPackage* Package = CreatePackage(NULL, *NewPackageName);
-	
+	UPackage* Package = UPackageTools::FindOrCreatePackageForAssetType(FName(*NewPackageName), UMaterial::StaticClass());
+	check(Package);
+	NewPackageName = Package->GetFullName();
+	FixedMaterialName = FPaths::GetBaseFilename(NewPackageName, true);
+
 	// does not override existing materials
 	UMaterialInterface* UnrealMaterialInterface = FindObject<UMaterialInterface>(Package, *FixedMaterialName);
 	if (UnrealMaterialInterface != NULL)
 	{
+		// Keep track of the processed materials
+		ImportContext.ImportedMaterials.Add(MaterialFullName, UnrealMaterialInterface);
+
 		// touch the textures anyway to make sure they reload if necessary
-		UTexture* DiffuseTexture = CreateSpeedTreeMaterialTexture(Parent, ANSI_TO_TCHAR(RenderState->m_apTextures[SpeedTree::TL_DIFFUSE]), false, false, LoadedPackages);
+		UTexture* DiffuseTexture = CreateSpeedTreeMaterialTexture(Parent, ANSI_TO_TCHAR(RenderState->m_apTextures[SpeedTree::TL_DIFFUSE]), false, false, LoadedPackages, ImportContext);
 		if (DiffuseTexture)
 		{
 			if (RenderState->m_bBranchesPresent && Options->SpeedTreeImportData->IncludeDetailMapCheck)
 			{
-				UTexture* DetailTexture = CreateSpeedTreeMaterialTexture(Parent, ANSI_TO_TCHAR(RenderState->m_apTextures[SpeedTree::TL_DETAIL_DIFFUSE]), false, false, LoadedPackages);
+				UTexture* DetailTexture = CreateSpeedTreeMaterialTexture(Parent, ANSI_TO_TCHAR(RenderState->m_apTextures[SpeedTree::TL_DETAIL_DIFFUSE]), false, false, LoadedPackages, ImportContext);
 			}
 		}
 		if (Options->SpeedTreeImportData->IncludeSpecularMapCheck)
 		{
-			UTexture* SpecularTexture = CreateSpeedTreeMaterialTexture(Parent, ANSI_TO_TCHAR(RenderState->m_apTextures[SpeedTree::TL_SPECULAR_MASK]), false, false, LoadedPackages);
+			UTexture* SpecularTexture = CreateSpeedTreeMaterialTexture(Parent, ANSI_TO_TCHAR(RenderState->m_apTextures[SpeedTree::TL_SPECULAR_MASK]), false, false, LoadedPackages, ImportContext);
 		}
 		if (Options->SpeedTreeImportData->IncludeNormalMapCheck)
 		{
-			UTexture* NormalTexture = CreateSpeedTreeMaterialTexture(Parent, ANSI_TO_TCHAR(RenderState->m_apTextures[SpeedTree::TL_NORMAL]), true, false, LoadedPackages);
+			UTexture* NormalTexture = CreateSpeedTreeMaterialTexture(Parent, ANSI_TO_TCHAR(RenderState->m_apTextures[SpeedTree::TL_NORMAL]), true, false, LoadedPackages, ImportContext);
 		}
 
 		return UnrealMaterialInterface;
@@ -575,6 +618,9 @@
 	FAssetRegistryModule::AssetCreated(UnrealMaterial);
 	Package->SetDirtyFlag(true);
 
+	// Keep track of the processed materials
+	ImportContext.ImportedMaterials.Add(MaterialFullName, UnrealMaterialInterface);
+
 	if (!RenderState->m_bDiffuseAlphaMaskIsOpaque && !RenderState->m_bBranchesPresent && !RenderState->m_bRigidMeshesPresent)
 	{
 		UnrealMaterial->BlendMode = BLEND_Masked;
@@ -608,7 +654,7 @@
 	}
 
 	// textures and properties
-	UTexture* DiffuseTexture = CreateSpeedTreeMaterialTexture(Parent, ANSI_TO_TCHAR(RenderState->m_apTextures[SpeedTree::TL_DIFFUSE]), false, false, LoadedPackages);
+	UTexture* DiffuseTexture = CreateSpeedTreeMaterialTexture(Parent, ANSI_TO_TCHAR(RenderState->m_apTextures[SpeedTree::TL_DIFFUSE]), false, false, LoadedPackages, ImportContext);
 	if (DiffuseTexture)
 	{
 		// make texture sampler
@@ -650,7 +696,7 @@
 
 		if (RenderState->m_bBranchesPresent && Options->SpeedTreeImportData->IncludeDetailMapCheck)
 		{
-			UTexture* DetailTexture = CreateSpeedTreeMaterialTexture(Parent, ANSI_TO_TCHAR(RenderState->m_apTextures[SpeedTree::TL_DETAIL_DIFFUSE]), false, false, LoadedPackages);
+			UTexture* DetailTexture = CreateSpeedTreeMaterialTexture(Parent, ANSI_TO_TCHAR(RenderState->m_apTextures[SpeedTree::TL_DETAIL_DIFFUSE]), false, false, LoadedPackages, ImportContext);
 			if (DetailTexture)
 			{
 				// add/find UVSet
@@ -686,7 +732,7 @@
 	bool bMadeSpecular = false;
 	if (Options->SpeedTreeImportData->IncludeSpecularMapCheck)
 	{
-		UTexture* SpecularTexture = CreateSpeedTreeMaterialTexture(Parent, ANSI_TO_TCHAR(RenderState->m_apTextures[SpeedTree::TL_SPECULAR_MASK]), false, false, LoadedPackages);
+		UTexture* SpecularTexture = CreateSpeedTreeMaterialTexture(Parent, ANSI_TO_TCHAR(RenderState->m_apTextures[SpeedTree::TL_SPECULAR_MASK]), false, false, LoadedPackages, ImportContext);
 		if (SpecularTexture)
 		{
 			// make texture sampler
@@ -710,7 +756,7 @@
 
 	if (Options->SpeedTreeImportData->IncludeNormalMapCheck)
 	{
-		UTexture* NormalTexture = CreateSpeedTreeMaterialTexture(Parent, ANSI_TO_TCHAR(RenderState->m_apTextures[SpeedTree::TL_NORMAL]), true, false, LoadedPackages);
+		UTexture* NormalTexture = CreateSpeedTreeMaterialTexture(Parent, ANSI_TO_TCHAR(RenderState->m_apTextures[SpeedTree::TL_NORMAL]), true, false, LoadedPackages, ImportContext);
 		if (NormalTexture)
 		{
 			// make texture sampler
@@ -845,7 +891,7 @@
 	return UnrealMaterial;
 }
 
-UMaterialInterface* CreateSpeedTreeMaterial8(UObject* Parent, FString MaterialFullName, SpeedTree8::CMaterial& SpeedTreeMaterial, TSharedPtr<SSpeedTreeImportOptions> Options, ESpeedTreeWindType WindType, ESpeedTreeGeometryType GeomType, TSet<UPackage*>& LoadedPackages, bool bCrossfadeLOD)
+UMaterialInterface* CreateSpeedTreeMaterial8(UObject* Parent, FString MaterialFullName, SpeedTree8::CMaterial& SpeedTreeMaterial, TSharedPtr<SSpeedTreeImportOptions> Options, ESpeedTreeWindType WindType, ESpeedTreeGeometryType GeomType, TSet<UPackage*>& LoadedPackages, bool bCrossfadeLOD, FSpeedTreeImportContext& ImportContext)
 {
 	// Make sure we have a parent
 	if (!Options->SpeedTreeImportData->MakeMaterialsCheck || !ensure(Parent))
@@ -853,30 +899,41 @@
 		return UMaterial::GetDefaultMaterial(MD_Surface);
 	}
 
+	if (UMaterialInterface** Material = ImportContext.ImportedMaterials.Find(MaterialFullName))
+	{
+		// The material was already imported
+		return *Material;
+	}
+
 	// set where to place the materials
 	FString FixedMaterialName = ObjectTools::SanitizeObjectName(MaterialFullName);
 	FString NewPackageName = FPackageName::GetLongPackagePath(Parent->GetOutermost()->GetName()) + TEXT("/") + FixedMaterialName;
-	NewPackageName = UPackageTools::SanitizePackageName(NewPackageName);
-	UPackage* Package = CreatePackage(NULL, *NewPackageName);
+	UPackage* Package = UPackageTools::FindOrCreatePackageForAssetType(FName(*NewPackageName), UMaterial::StaticClass());
+	check(Package);
+	NewPackageName = Package->GetFullName();
+	FixedMaterialName = FPaths::GetBaseFilename(NewPackageName, true);
 
 	// does not override existing materials
 	UMaterialInterface* UnrealMaterialInterface = FindObject<UMaterialInterface>(Package, *FixedMaterialName);
 	if (UnrealMaterialInterface != NULL)
 	{
+		// Keep track of the processed materials
+		ImportContext.ImportedMaterials.Add(MaterialFullName, UnrealMaterialInterface);
+
 		// touch the textures anyway to make sure they reload if necessary
 		if (SpeedTreeMaterial.Maps().Count() > 0 && SpeedTreeMaterial.Maps()[0].Used() && !SpeedTreeMaterial.Maps()[0].Path().IsEmpty())
 		{
-			UTexture* DiffuseTexture = CreateSpeedTreeMaterialTexture(Parent, ANSI_TO_TCHAR(SpeedTreeMaterial.Maps()[0].Path().Data()), false, false, LoadedPackages);
+			UTexture* DiffuseTexture = CreateSpeedTreeMaterialTexture(Parent, ANSI_TO_TCHAR(SpeedTreeMaterial.Maps()[0].Path().Data()), false, false, LoadedPackages, ImportContext);
 		}
 
 		if (SpeedTreeMaterial.Maps().Count() > 1 && SpeedTreeMaterial.Maps()[1].Used() && !SpeedTreeMaterial.Maps()[1].Path().IsEmpty())
 		{
-			UTexture* NormalTexture = CreateSpeedTreeMaterialTexture(Parent, ANSI_TO_TCHAR(SpeedTreeMaterial.Maps()[1].Path().Data()), false, false, LoadedPackages);
+			UTexture* NormalTexture = CreateSpeedTreeMaterialTexture(Parent, ANSI_TO_TCHAR(SpeedTreeMaterial.Maps()[1].Path().Data()), false, false, LoadedPackages, ImportContext);
 		}
 
 		if (Options->SpeedTreeImportData->IncludeSubsurface && SpeedTreeMaterial.Maps().Count() > 2 && SpeedTreeMaterial.Maps()[2].Used() && !SpeedTreeMaterial.Maps()[2].Path().IsEmpty())
 		{
-			UTexture* SubsurfaceTexture = CreateSpeedTreeMaterialTexture(Parent, ANSI_TO_TCHAR(SpeedTreeMaterial.Maps()[2].Path().Data()), false, false, LoadedPackages);
+			UTexture* SubsurfaceTexture = CreateSpeedTreeMaterialTexture(Parent, ANSI_TO_TCHAR(SpeedTreeMaterial.Maps()[2].Path().Data()), false, false, LoadedPackages, ImportContext);
 		}
 
 		return UnrealMaterialInterface;
@@ -888,6 +945,9 @@
 	FAssetRegistryModule::AssetCreated(UnrealMaterial);
 	Package->SetDirtyFlag(true);
 
+	// Keep track of the processed materials
+	ImportContext.ImportedMaterials.Add(MaterialFullName, UnrealMaterialInterface);
+
 	// basic setup
 	UnrealMaterial->TwoSided = SpeedTreeMaterial.TwoSided();
 	UnrealMaterial->BlendMode = BLEND_Masked;
@@ -934,7 +994,7 @@
 		}
 		else
 		{
-			UTexture* DiffuseTexture = CreateSpeedTreeMaterialTexture(Parent, ANSI_TO_TCHAR(SpeedTreeMaterial.Maps()[0].Path().Data()), false, false, LoadedPackages);
+			UTexture* DiffuseTexture = CreateSpeedTreeMaterialTexture(Parent, ANSI_TO_TCHAR(SpeedTreeMaterial.Maps()[0].Path().Data()), false, false, LoadedPackages, ImportContext);
 			if (DiffuseTexture)
 			{
 				// this helps prevent mipmapping from eating away tiny leaves
@@ -1008,7 +1068,7 @@
 		{
 			// default SpeedTree v8 texture packing uses one texture for normal and roughness. BC5 doesn't support alpha, so compress with default settings, but disable sRGB. Then bias the normal ourselves
 
-			UTexture* NormalTexture = CreateSpeedTreeMaterialTexture(Parent, ANSI_TO_TCHAR(SpeedTreeMaterial.Maps()[1].Path().Data()), false, false, LoadedPackages);
+			UTexture* NormalTexture = CreateSpeedTreeMaterialTexture(Parent, ANSI_TO_TCHAR(SpeedTreeMaterial.Maps()[1].Path().Data()), false, false, LoadedPackages, ImportContext);
 			if (NormalTexture)
 			{
 				NormalTexture->SRGB = false;
@@ -1076,7 +1136,7 @@
 		}
 		else
 		{
-			UTexture* SubsurfaceTexture = CreateSpeedTreeMaterialTexture(Parent, ANSI_TO_TCHAR(SpeedTreeMaterial.Maps()[2].Path().Data()), false, false, LoadedPackages);
+			UTexture* SubsurfaceTexture = CreateSpeedTreeMaterialTexture(Parent, ANSI_TO_TCHAR(SpeedTreeMaterial.Maps()[2].Path().Data()), false, false, LoadedPackages, ImportContext);
 			if (SubsurfaceTexture)
 			{
 				// make texture sampler
@@ -1497,8 +1557,10 @@
 
 	FString MeshName = ObjectTools::SanitizeObjectName(InName.ToString());
 	FString NewPackageName = FPackageName::GetLongPackagePath(InParent->GetOutermost()->GetName()) + TEXT("/") + MeshName;
-	NewPackageName = UPackageTools::SanitizePackageName(NewPackageName);
-	UPackage* Package = CreatePackage(NULL, *NewPackageName);
+	UPackage* Package = UPackageTools::FindOrCreatePackageForAssetType(FName(*NewPackageName), UStaticMesh::StaticClass());
+	check(Package);
+	NewPackageName = Package->GetFullName();
+	MeshName = FPaths::GetBaseFilename(NewPackageName, true);
 
 	UStaticMesh* ExistingMesh = FindObject<UStaticMesh>(Package, *MeshName);
 	USpeedTreeImportData* ExistingImportData = nullptr;
@@ -1535,6 +1597,8 @@
 		}
 		else
 		{
+			FSpeedTreeImportContext ImportContext;
+
 			const SpeedTree::SGeometry* SpeedTreeGeometry = SpeedTree.GetGeometry();
 			if ((Options->SpeedTreeImportData->ImportGeometryType == EImportGeometryType::IGT_Billboards && SpeedTreeGeometry->m_sVertBBs.m_nNumBillboards == 0) ||
 				(Options->SpeedTreeImportData->ImportGeometryType == EImportGeometryType::IGT_3D && SpeedTreeGeometry->m_nNumLods == 0))
@@ -1737,7 +1801,7 @@
 
 								MaterialName = ObjectTools::SanitizeObjectName(MaterialName);
 
-								UMaterialInterface* Material = CreateSpeedTreeMaterial7(InParent, MaterialName, RenderState, Options, WindType, SpeedTreeGeometry->m_sVertBBs.m_nNumBillboards, LoadedPackages);
+								UMaterialInterface* Material = CreateSpeedTreeMaterial7(InParent, MaterialName, RenderState, Options, WindType, SpeedTreeGeometry->m_sVertBBs.m_nNumBillboards, LoadedPackages, ImportContext);
 								
 								RenderStateIndexToStaticMeshIndex.Add(DrawCall->m_nRenderStateIndex, StaticMesh->StaticMaterials.Num());
 								MaterialIndex = StaticMesh->StaticMaterials.Add(FStaticMaterial(Material, FName(*MaterialName), FName(*MaterialName)));
@@ -1845,7 +1909,7 @@
 					}
 
 					FString MaterialName = MeshName + "_Billboard";
-					UMaterialInterface* Material = CreateSpeedTreeMaterial7(InParent, MaterialName, &SpeedTreeGeometry->m_aBillboardRenderStates[SpeedTree::RENDER_PASS_MAIN], Options, WindType, SpeedTreeGeometry->m_sVertBBs.m_nNumBillboards, LoadedPackages);
+					UMaterialInterface* Material = CreateSpeedTreeMaterial7(InParent, MaterialName, &SpeedTreeGeometry->m_aBillboardRenderStates[SpeedTree::RENDER_PASS_MAIN], Options, WindType, SpeedTreeGeometry->m_sVertBBs.m_nNumBillboards, LoadedPackages, ImportContext);
 					int32 MaterialIndex = StaticMesh->StaticMaterials.Add(FStaticMaterial(Material, FName(*MaterialName), FName(*MaterialName)));
 
 					const FPolygonGroupID CurrentPolygonGroupID = MeshDescription->CreatePolygonGroup();
@@ -2013,8 +2077,9 @@
 
 	FString MeshName = ObjectTools::SanitizeObjectName(InName.ToString());
 	FString NewPackageName = FPackageName::GetLongPackagePath(InParent->GetOutermost()->GetName()) + TEXT("/") + MeshName;
-	NewPackageName = UPackageTools::SanitizePackageName(NewPackageName);
-	UPackage* Package = CreatePackage(NULL, *NewPackageName);
+	UPackage* Package = UPackageTools::FindOrCreatePackageForAssetType(FName(*NewPackageName), UStaticMesh::StaticClass());
+	check(Package);
+	NewPackageName = Package->GetFullName();
 
 	UStaticMesh* ExistingMesh = FindObject<UStaticMesh>(Package, *MeshName);
 	USpeedTreeImportData* ExistingImportData = nullptr;
@@ -2043,6 +2108,8 @@
 	{
 		LoadedPackages.Empty();
 
+		FSpeedTreeImportContext ImportContext;
+
 		// clear out old mesh
 		TArray<FStaticMaterial> OldMaterials;
 		FGlobalComponentReregisterContext RecreateComponents;
@@ -2205,7 +2272,7 @@
 					SpeedTree8::CMaterial SpeedTreeMaterial = SpeedTree.Materials()[DrawCall.m_uiMaterialIndex];
 					FString MaterialName = FString(SpeedTreeMaterial.Name().Data());
 					MaterialName.InsertAt(MaterialName.Len() - 4, GeomString);
-					UMaterialInterface* Material = CreateSpeedTreeMaterial8(InParent, MaterialName, SpeedTreeMaterial, Options, WindType, GeomType, LoadedPackages, bCrossfadeLOD);
+					UMaterialInterface* Material = CreateSpeedTreeMaterial8(InParent, MaterialName, SpeedTreeMaterial, Options, WindType, GeomType, LoadedPackages, bCrossfadeLOD, ImportContext);
 					MaterialMap.Add(MaterialKey, StaticMesh->StaticMaterials.Num());
 					int32 AddedMaterialIndex = StaticMesh->StaticMaterials.Add(FStaticMaterial(Material, FName(*MaterialName), FName(*MaterialName)));
 					const FPolygonGroupID& PolygonGroupID = MeshDescription->CreatePolygonGroup();
diff --git a/Engine/Source/Editor/UnrealEd/Private/PackageTools.cpp b/Engine/Source/Editor/UnrealEd/Private/PackageTools.cpp
--- a/Engine/Source/Editor/UnrealEd/Private/PackageTools.cpp
+++ b/Engine/Source/Editor/UnrealEd/Private/PackageTools.cpp
@@ -39,15 +39,18 @@
 #include "Framework/Notifications/NotificationManager.h"
 #include "Widgets/Notifications/SNotificationList.h"
 
+#include "AssetData.h"
 #include "AssetRegistryModule.h"
-#include "Logging/MessageLog.h"
-#include "UObject/UObjectIterator.h"
 #include "ComponentReregisterContext.h"
 #include "Engine/BlueprintGeneratedClass.h"
 #include "Engine/GameEngine.h"
 #include "Engine/LevelStreaming.h"
 #include "Engine/MapBuildDataRegistry.h"
 #include "Engine/Selection.h"
+#include "IAssetRegistry.h"
+#include "Logging/MessageLog.h"
+#include "UObject/UObjectGlobals.h"
+#include "UObject/UObjectIterator.h"
 
 #include "ShaderCompiler.h"
 #include "DistanceFieldAtlas.h"
@@ -1185,6 +1188,65 @@ UPackageTools::UPackageTools(const FObjectInitializer& ObjectInitializer)
 		return SanitizedName;
 	}
 
+	UPackage* UPackageTools::FindOrCreatePackageForAssetType(const FName LongPackageName, UClass* AssetClass)
+	{
+		if (AssetClass)
+		{
+			// Test the asset registry first
+			// Only scan the disk cache since it's fast O(1) (otherwise the asset registry will iterate on all the objects in memory)
+			constexpr bool bIncludeOnlyOnDiskAssets = true;
+			TArray<FAssetData> OutAssets;
+			IAssetRegistry& AssetRegistry = FAssetRegistryModule::GetRegistry();
+			AssetRegistry.GetAssetsByPackageName(LongPackageName, OutAssets, bIncludeOnlyOnDiskAssets);
+			bool bShouldGenerateUniquePackageName = false;
+			int32 NumberOfAssets = OutAssets.Num();
+			if (NumberOfAssets == 1)
+			{
+				const FAssetData& AssetData = OutAssets[0];
+				if (AssetData.AssetClass != AssetClass->GetFName())
+				{
+					bShouldGenerateUniquePackageName = true;
+				}
+			}
+			else if (NumberOfAssets > 1)
+			{
+				// this shouldn't happen in UE4
+				bShouldGenerateUniquePackageName = true;
+			}
+
+			UPackage* Package = nullptr;
+			if (!bShouldGenerateUniquePackageName)
+			{
+				// Create or return the existing package
+				FString PackageNameAsString = LongPackageName.ToString();
+				Package = LoadPackage(PackageNameAsString);
+
+				if (Package)
+				{
+					UObject* Object = FindObjectFast<UObject>(Package, Package->GetFName());
+					if (Object && Object->GetClass() != AssetClass)
+					{
+						bShouldGenerateUniquePackageName = true;
+					}
+				}
+				else
+				{
+					Package = NewObject<UPackage>(nullptr, LongPackageName, RF_Public);
+				}
+			}
+
+			if (bShouldGenerateUniquePackageName)
+			{
+				FString NewPackageName = MakeUniqueObjectName(nullptr, UPackage::StaticClass(), LongPackageName).ToString();
+				Package = NewObject<UPackage>(nullptr, *NewPackageName, RF_Public);
+			}
+
+			return Package;
+		}
+
+		return nullptr;
+	}
+
 #undef LOCTEXT_NAMESPACE
 
 // EOF
diff --git a/Engine/Source/Editor/UnrealEd/Public/PackageTools.h b/Engine/Source/Editor/UnrealEd/Public/PackageTools.h
--- a/Engine/Source/Editor/UnrealEd/Public/PackageTools.h
+++ b/Engine/Source/Editor/UnrealEd/Public/PackageTools.h
@@ -178,6 +178,15 @@ public:
 	UFUNCTION(BlueprintCallable, Category = "Editor Scripting | Package Tools")
 	static FString SanitizePackageName(const FString& InPackageName);
 
+	/**
+	 * Find or create a package at the desired path
+	 * If a package already exist but was made for another type of asset,
+	 * this function will create another one with a modified unique name
+	 * @param LongPackageName The package path in unreal asset path
+	 * @param AssetClass The class of asset the package should be for.
+	 */
+	static UPackage* FindOrCreatePackageForAssetType(const FName LongPackageName, UClass* AssetClass);
+
 private:
 	static void RestoreStandaloneOnReachableObjects();
 
@@ -190,4 +199,3 @@ private:
 
 UE_DEPRECATED(4.21, "PackageTools namespace has been deprecated. Please use UPackageTools instead.") 
 typedef UPackageTools PackageTools;
-
