Fixed mesh attribute deprecation warnings.

diff --git a/Engine/Plugins/Editor/SpeedTreeImporter/Source/SpeedTreeImporter/Private/SpeedTreeImportFactory.cpp b/Engine/Plugins/Editor/SpeedTreeImporter/Source/SpeedTreeImporter/Private/SpeedTreeImportFactory.cpp
--- a/Engine/Plugins/Editor/SpeedTreeImporter/Source/SpeedTreeImporter/Private/SpeedTreeImportFactory.cpp
+++ b/Engine/Plugins/Editor/SpeedTreeImporter/Source/SpeedTreeImporter/Private/SpeedTreeImportFactory.cpp
@@ -2882,7 +2882,6 @@ UObject* USpeedTreeImportFactory::FactoryCreateBinary9(UClass* InClass, UObject*
 
 			TVertexAttributesRef<FVector> VertexPositions = Attributes.GetVertexPositions();
 			TEdgeAttributesRef<bool> EdgeHardnesses = Attributes.GetEdgeHardnesses();
-			TEdgeAttributesRef<float> EdgeCreaseSharpnesses = Attributes.GetEdgeCreaseSharpnesses();
 			TPolygonGroupAttributesRef<FName> PolygonGroupImportedMaterialSlotNames = Attributes.GetPolygonGroupMaterialSlotNames();
 			TVertexInstanceAttributesRef<FVector> VertexInstanceNormals = Attributes.GetVertexInstanceNormals();
 			TVertexInstanceAttributesRef<FVector> VertexInstanceTangents = Attributes.GetVertexInstanceTangents();
@@ -2890,7 +2889,7 @@ UObject* USpeedTreeImportFactory::FactoryCreateBinary9(UClass* InClass, UObject*
 			TVertexInstanceAttributesRef<FVector4> VertexInstanceColors = Attributes.GetVertexInstanceColors();
 			TVertexInstanceAttributesRef<FVector2D> VertexInstanceUVs = Attributes.GetVertexInstanceUVs();
 
-			VertexInstanceUVs.SetNumIndices(NumUVs);
+			VertexInstanceUVs.SetNumChannels(NumUVs);
 
 			for (int32 MatIndex = 0; MatIndex < StaticMesh->StaticMaterials.Num(); ++MatIndex)
 			{
