Fix texture import in SpeedTree to properly rebuild the texture if properties are modified after the initial import

diff --git a/Engine/Plugins/Editor/SpeedTreeImporter/Source/SpeedTreeImporter/Private/SpeedTreeImportFactory.cpp b/Engine/Plugins/Editor/SpeedTreeImporter/Source/SpeedTreeImporter/Private/SpeedTreeImportFactory.cpp
--- a/Engine/Plugins/Editor/SpeedTreeImporter/Source/SpeedTreeImporter/Private/SpeedTreeImportFactory.cpp
+++ b/Engine/Plugins/Editor/SpeedTreeImporter/Source/SpeedTreeImporter/Private/SpeedTreeImportFactory.cpp
@@ -493,12 +493,17 @@ UTexture* CreateSpeedTreeMaterialTexture(UObject* Parent,const FString& Filename
 
 	const uint8* PtrTexture = TextureData.GetData();
 	UTexture* Texture = (UTexture*)TextureFact->FactoryCreateBinary(UTexture2D::StaticClass(), Package, *TextureName, RF_Standalone|RF_Public, NULL, *Extension, PtrTexture, PtrTexture + TextureData.Num(), GWarn);
-	if (Texture != NULL)
+	if (Texture != nullptr)
 	{
 		if (bMasks)
 		{
-			Texture->SRGB = false;
-			Texture->CompressionSettings = TC_Masks;
+			if (Texture->SRGB != false || Texture->CompressionSettings != TC_Masks)
+			{
+				Texture->PreEditChange(nullptr);
+				Texture->SRGB = false;
+				Texture->CompressionSettings = TC_Masks;
+				Texture->PostEditChange();
+			}
 		}
 		Texture->AssetImportData->Update(FilenamePath);
 
@@ -1068,7 +1073,12 @@ UMaterialInterface* CreateSpeedTreeMaterial8(UObject* Parent, FString MaterialFu
 			if (DiffuseTexture)
 			{
 				// this helps prevent mipmapping from eating away tiny leaves
-				DiffuseTexture->AdjustMinAlpha = 0.1f;
+				if (!FMath::IsNearlyEqual(DiffuseTexture->AdjustMinAlpha, 0.1f))
+				{
+					DiffuseTexture->PreEditChange(nullptr);
+					DiffuseTexture->AdjustMinAlpha = 0.1f;
+					DiffuseTexture->PostEditChange();
+				}
 
 				// make texture sampler
 				UMaterialExpressionTextureSample* TextureExpression = NewObject<UMaterialExpressionTextureSample>(UnrealMaterial);
@@ -1141,7 +1151,12 @@ UMaterialInterface* CreateSpeedTreeMaterial8(UObject* Parent, FString MaterialFu
 			UTexture* NormalTexture = CreateSpeedTreeMaterialTexture(Parent, ANSI_TO_TCHAR(SpeedTreeMaterial.Maps()[1].Path().Data()), false, false, LoadedPackages, ImportContext);
 			if (NormalTexture)
 			{
-				NormalTexture->SRGB = false;
+				if (NormalTexture->SRGB != false)
+				{
+					NormalTexture->PreEditChange(nullptr);
+					NormalTexture->SRGB = false;
+					NormalTexture->PostEditChange();
+				}
 
 				// make texture sampler
 				UMaterialExpressionTextureSample* TextureExpression = NewObject<UMaterialExpressionTextureSample>(UnrealMaterial);
@@ -1434,7 +1449,12 @@ UMaterialInterface* CreateSpeedTreeMaterial9(UObject* Parent, FString MaterialFu
 		if (Texture)
 		{
 			// this helps prevent mipmapping from eating away tiny leaves
-			Texture->AdjustMinAlpha = 0.05f;
+			if (!FMath::IsNearlyEqual(Texture->AdjustMinAlpha, 0.05f))
+			{
+				Texture->PreEditChange(nullptr);
+				Texture->AdjustMinAlpha = 0.05f;
+				Texture->PostEditChange();
+			}
 
 			UnrealMaterialInstance->SetTextureParameterValueEditorOnly(FMaterialParameterInfo("ColorOpacity"), Texture);
 		}
@@ -1445,7 +1465,12 @@ UMaterialInterface* CreateSpeedTreeMaterial9(UObject* Parent, FString MaterialFu
 		UTexture* Texture = CreateSpeedTreeMaterialTexture(Parent, ANSI_TO_TCHAR(SpeedTreeMaterial.Maps()[1].Path().Data()), false, false, LoadedPackages, ImportContext);
 		if (Texture)
 		{
-			Texture->SRGB = false;
+			if (Texture->SRGB != false)
+			{
+				Texture->PreEditChange(nullptr);
+				Texture->SRGB = false;
+				Texture->PostEditChange();
+			}
 
 			UnrealMaterialInstance->SetTextureParameterValueEditorOnly(FMaterialParameterInfo("NormalRoughness"), Texture);
 		}
