The reset option button on the speed tree import dialog now resets to  thedefault settings instead of the settings of the last import.

diff --git a/Engine/Plugins/Editor/SpeedTreeImporter/Source/SpeedTreeImporter/Classes/SpeedTreeImportData.h b/Engine/Plugins/Editor/SpeedTreeImporter/Source/SpeedTreeImporter/Classes/SpeedTreeImportData.h
--- a/Engine/Plugins/Editor/SpeedTreeImporter/Source/SpeedTreeImporter/Classes/SpeedTreeImportData.h
+++ b/Engine/Plugins/Editor/SpeedTreeImporter/Source/SpeedTreeImporter/Classes/SpeedTreeImportData.h
@@ -41,7 +41,8 @@ enum EImportLODType
 	ILT_IndividualActors UMETA(DisplayName = "Individual Actors")
 };
 
-UCLASS(BlueprintType, config=EditorPerProjectUserSettings, MinimalAPI)
+
+UCLASS(BlueprintType, config=EditorPerProjectUserSettings, perObjectConfig, MinimalAPI)
 class USpeedTreeImportData : public UAssetImportData
 {
 	GENERATED_UCLASS_BODY()
@@ -107,10 +108,6 @@ class USpeedTreeImportData : public UAssetImportData
 	uint32 IncludeSmoothLODCheck : 1;
 
 	void CopyFrom(USpeedTreeImportData* Other);
-
-	//Save and load options to retrieve the last option that was use when importing a speedtree asset
-	void SaveOptions();
-	void LoadOptions();
 };
 
 
diff --git a/Engine/Plugins/Editor/SpeedTreeImporter/Source/SpeedTreeImporter/Private/SpeedTreeImportData.cpp b/Engine/Plugins/Editor/SpeedTreeImporter/Source/SpeedTreeImporter/Private/SpeedTreeImportData.cpp
--- a/Engine/Plugins/Editor/SpeedTreeImporter/Source/SpeedTreeImporter/Private/SpeedTreeImportData.cpp
+++ b/Engine/Plugins/Editor/SpeedTreeImporter/Source/SpeedTreeImporter/Private/SpeedTreeImportData.cpp
@@ -53,150 +53,6 @@ void USpeedTreeImportData::CopyFrom(USpeedTreeImportData* Other)
 	IncludeSmoothLODCheck = Other->IncludeSmoothLODCheck;
 }
 
-void USpeedTreeImportData::SaveOptions()
-{
-	int32 PortFlags = 0;
-
-	for (FProperty* Property = GetClass()->PropertyLink; Property; Property = Property->PropertyLinkNext)
-	{
-		if (!Property->HasAnyPropertyFlags(CPF_Config))
-		{
-			continue;
-		}
-		FString Section = TEXT("SpeedTree_Import_UI_Option_") + GetClass()->GetName();
-		FString Key = Property->GetName();
-
-		const bool bIsPropertyInherited = Property->GetOwnerClass() != GetClass();
-		UObject* SuperClassDefaultObject = GetClass()->GetSuperClass()->GetDefaultObject();
-
-		FArrayProperty* Array = CastField<FArrayProperty>(Property);
-		if (Array)
-		{
-			FConfigSection* Sec = GConfig->GetSectionPrivate(*Section, 1, 0, *GEditorPerProjectIni);
-			check(Sec);
-			Sec->Remove(*Key);
-
-			FScriptArrayHelper_InContainer ArrayHelper(Array, this);
-			for (int32 i = 0; i < ArrayHelper.Num(); i++)
-			{
-				FString	Buffer;
-				Array->Inner->ExportTextItem(Buffer, ArrayHelper.GetRawPtr(i), ArrayHelper.GetRawPtr(i), this, PortFlags);
-				Sec->Add(*Key, *Buffer);
-			}
-		}
-		else
-		{
-			TCHAR TempKey[MAX_SPRINTF] = TEXT("");
-			for (int32 Index = 0; Index < Property->ArrayDim; Index++)
-			{
-				if (Property->ArrayDim != 1)
-				{
-					FCString::Sprintf(TempKey, TEXT("%s[%i]"), *Property->GetName(), Index);
-					Key = TempKey;
-				}
-
-				FString	Value;
-				Property->ExportText_InContainer(Index, Value, this, this, this, PortFlags);
-				GConfig->SetString(*Section, *Key, *Value, *GEditorPerProjectIni);
-			}
-		}
-	}
-	GConfig->Flush(0);
-}
-
-void USpeedTreeImportData::LoadOptions()
-{
-	int32 PortFlags = 0;
-
-	for (FProperty* Property = GetClass()->PropertyLink; Property; Property = Property->PropertyLinkNext)
-	{
-		if (!Property->HasAnyPropertyFlags(CPF_Config))
-		{
-			continue;
-		}
-		FString Section = TEXT("SpeedTree_Import_UI_Option_") + GetClass()->GetName();
-		FString Key = Property->GetName();
-
-		const bool bIsPropertyInherited = Property->GetOwnerClass() != GetClass();
-		UObject* SuperClassDefaultObject = GetClass()->GetSuperClass()->GetDefaultObject();
-
-		const FString& PropFileName = GEditorPerProjectIni;
-
-		FArrayProperty* Array = CastField<FArrayProperty>(Property);
-		if (Array)
-		{
-			FConfigSection* Sec = GConfig->GetSectionPrivate(*Section, 0, 1, *GEditorPerProjectIni);
-			if (Sec != nullptr)
-			{
-				TArray<FConfigValue> List;
-				const FName KeyName(*Key, FNAME_Find);
-				Sec->MultiFind(KeyName, List);
-
-				FScriptArrayHelper_InContainer ArrayHelper(Array, this);
-				// Only override default properties if there is something to override them with.
-				if (List.Num() > 0)
-				{
-					ArrayHelper.EmptyAndAddValues(List.Num());
-					for (int32 i = List.Num() - 1, c = 0; i >= 0; i--, c++)
-					{
-						Array->Inner->ImportText(*List[i].GetValue(), ArrayHelper.GetRawPtr(c), PortFlags, this);
-					}
-				}
-				else
-				{
-					int32 Index = 0;
-					const FConfigValue* ElementValue = nullptr;
-					do
-					{
-						// Add array index number to end of key
-						FString IndexedKey = FString::Printf(TEXT("%s[%i]"), *Key, Index);
-
-						// Try to find value of key
-						const FName IndexedName(*IndexedKey, FNAME_Find);
-						if (IndexedName == NAME_None)
-						{
-							break;
-						}
-						ElementValue = Sec->Find(IndexedName);
-
-						// If found, import the element
-						if (ElementValue != nullptr)
-						{
-							// expand the array if necessary so that Index is a valid element
-							ArrayHelper.ExpandForIndex(Index);
-							Array->Inner->ImportText(*ElementValue->GetValue(), ArrayHelper.GetRawPtr(Index), PortFlags, this);
-						}
-
-						Index++;
-					} while (ElementValue || Index < ArrayHelper.Num());
-				}
-			}
-		}
-		else
-		{
-			for (int32 i = 0; i < Property->ArrayDim; i++)
-			{
-				if (Property->ArrayDim != 1)
-				{
-					Key = FString::Printf(TEXT("%s[%i]"), *Property->GetName(), i);
-				}
-
-				FString Value;
-				bool bFoundValue = GConfig->GetString(*Section, *Key, Value, *GEditorPerProjectIni);
-
-				if (bFoundValue)
-				{
-					if (Property->ImportText(*Value, Property->ContainerPtrToValuePtr<uint8>(this, i), PortFlags, this) == NULL)
-					{
-						// this should be an error as the properties from the .ini / .int file are not correctly being read in and probably are affecting things in subtle ways
-						UE_LOG(LogSpeedTreeImportData, Error, TEXT("SpeedTree Options LoadOptions (%s): failed for %s in: %s"), *GetPathName(), *Property->GetName(), *Value);
-					}
-				}
-			}
-		}
-	}
-}
-
 FSpeedTreeImportDataDetails::FSpeedTreeImportDataDetails()
 {
 	SpeedTreeImportData = nullptr;
diff --git a/Engine/Plugins/Editor/SpeedTreeImporter/Source/SpeedTreeImporter/Private/SpeedTreeImportFactory.cpp b/Engine/Plugins/Editor/SpeedTreeImporter/Source/SpeedTreeImporter/Private/SpeedTreeImportFactory.cpp
--- a/Engine/Plugins/Editor/SpeedTreeImporter/Source/SpeedTreeImporter/Private/SpeedTreeImportFactory.cpp
+++ b/Engine/Plugins/Editor/SpeedTreeImporter/Source/SpeedTreeImporter/Private/SpeedTreeImportFactory.cpp
@@ -7,6 +7,7 @@
 #include "EditorFramework/AssetImportData.h"
 #include "EditorReimportHandler.h"
 #include "EditorStyleSet.h"
+#include "Engine/Engine.h"
 #include "Engine/StaticMesh.h"
 #include "Engine/Texture.h"
 #include "Engine/Texture2D.h"
@@ -151,8 +152,14 @@ public:
 		bImport(false)
 	{
 		DetailsView = nullptr;
-		SpeedTreeImportData = NewObject<USpeedTreeImportData>(GetTransientPackage(), NAME_None);
-		SpeedTreeImportData->LoadConfig();
+		
+		// The name options is important for the config storage
+		FName SpeedTreeImportDataName = TEXT("SpeedTreeImportData");
+		if (UObject* ObjectToMove = FindObjectFast<UObject>(GetTransientPackage(), SpeedTreeImportDataName))
+		{
+			ObjectToMove->Rename();
+		}
+		SpeedTreeImportData = NewObject<USpeedTreeImportData>(GetTransientPackage(), SpeedTreeImportDataName);
 	}
 
 	void Construct(const FArguments& InArgs)
@@ -169,7 +176,7 @@ public:
 		else
 		{
 			//When simply importing we load the local config file of the user so he rerieve the last import options
-			SpeedTreeImportData->LoadOptions();
+			SpeedTreeImportData->LoadConfig();
 		}
 
 		// set the filename now so the dialog can tell if it is SpeedTree 7 or 8
@@ -252,6 +259,8 @@ public:
 		return bImport;
 	}
 
+private:
+
 	/** Called when 'OK' button is pressed */
 	FReply OnImport()
 	{
@@ -265,7 +274,22 @@ public:
 	{
 		if (DetailsView.IsValid())
 		{
-			SpeedTreeImportData->LoadConfig();
+			// Reset values from the CDO 
+			UClass* Class = SpeedTreeImportData->GetClass();
+			UObject* CDO = Class->GetDefaultObject();
+
+			for (FProperty* Property = Class->PropertyLink; Property; Property = Property->PropertyLinkNext)
+			{
+				// Only reset the property that would have been store in the config
+				if (Property->HasAnyPropertyFlags(CPF_Config))
+				{
+					void* Dest = Property->ContainerPtrToValuePtr<void>(SpeedTreeImportData);
+					const void* Source = Property->ContainerPtrToValuePtr<void>(CDO);
+					Property->CopyCompleteValue(Dest, Source);
+				}
+			}
+
+
 			DetailsView->SetObject(SpeedTreeImportData, true);
 		}
 		return FReply::Handled();
