Added missing arguments to the functions save config when saving the last import setting of a speedtree file. By default it updated the CDO. This not desired as we use the CDO to reset the options.

diff --git a/Engine/Plugins/Editor/SpeedTreeImporter/Source/SpeedTreeImporter/Private/SpeedTreeImportFactory.cpp b/Engine/Plugins/Editor/SpeedTreeImporter/Source/SpeedTreeImporter/Private/SpeedTreeImportFactory.cpp
--- a/Engine/Plugins/Editor/SpeedTreeImporter/Source/SpeedTreeImporter/Private/SpeedTreeImportFactory.cpp
+++ b/Engine/Plugins/Editor/SpeedTreeImporter/Source/SpeedTreeImporter/Private/SpeedTreeImportFactory.cpp
@@ -3,6 +3,7 @@
 #include "SpeedTreeImportFactory.h"
 
 #include "AssetImportTask.h"
+#include "CoreGlobals.h"
 #include "Editor.h"
 #include "EditorFramework/AssetImportData.h"
 #include "EditorReimportHandler.h"
@@ -266,7 +267,7 @@ private:
 	{
 		bImport = true;
 		WidgetWindow->RequestDestroyWindow();
-		SpeedTreeImportData->SaveConfig();
+		SpeedTreeImportData->SaveConfig(CPF_Config, nullptr, GConfig, false);
 		return FReply::Handled();
 	}
 
