Removed a reference to UE4 in a comment.

diff --git a/Engine/Plugins/Editor/SpeedTreeImporter/Source/SpeedTreeImporter/Private/SpeedTreeImportFactory.cpp b/Engine/Plugins/Editor/SpeedTreeImporter/Source/SpeedTreeImporter/Private/SpeedTreeImportFactory.cpp
--- a/Engine/Plugins/Editor/SpeedTreeImporter/Source/SpeedTreeImporter/Private/SpeedTreeImportFactory.cpp
+++ b/Engine/Plugins/Editor/SpeedTreeImporter/Source/SpeedTreeImporter/Private/SpeedTreeImportFactory.cpp
@@ -929,7 +929,7 @@ UMaterialInterface* CreateSpeedTreeMaterial7(UObject* Parent, FString MaterialFu
 		UnrealMaterial->AmbientOcclusion.MaskA = 0;
 	}
 
-	// UE4 flips normals for two-sided materials. SpeedTrees don't need that
+	// Unreal flips normals for two-sided materials. SpeedTrees don't need that
 	if (UnrealMaterial->TwoSided)
 	{
 		UMaterialExpressionTwoSidedSign* TwoSidedSign = NewObject<UMaterialExpressionTwoSidedSign>(UnrealMaterial);
