Fix FVector2D variant casts - Plugins
Submitted on behalf of fred.kimberley

diff --git a/Engine/Plugins/2D/Paper2D/Source/Paper2D/Private/GroupedSpriteSceneProxy.cpp b/Engine/Plugins/2D/Paper2D/Source/Paper2D/Private/GroupedSpriteSceneProxy.cpp
--- a/Engine/Plugins/2D/Paper2D/Source/Paper2D/Private/GroupedSpriteSceneProxy.cpp
+++ b/Engine/Plugins/2D/Paper2D/Source/Paper2D/Private/GroupedSpriteSceneProxy.cpp
@@ -139,7 +139,7 @@ FGroupedSpriteSceneProxy::FGroupedSpriteSceneProxy(UPaperGroupedSpriteComponent*
 			{
 				const FVector LocalPos((PaperAxisX * SourceVert.X) + (PaperAxisY * SourceVert.Y));
 				const FVector ComponentSpacePos = InstanceData.Transform.TransformPosition(LocalPos);
-				const FVector2D UV(SourceVert.Z, SourceVert.W);
+				const FVector2f UV(SourceVert.Z, SourceVert.W);	// LWC_TODO: Precision loss
 
 				*VertexWritePtr++ = FDynamicMeshVertex(ComponentSpacePos, TangentX.ToFVector(), TangentZ.ToFVector(), UV, VertColor);
 			}
diff --git a/Engine/Plugins/2D/Paper2D/Source/Paper2D/Private/PaperRenderSceneProxy.h b/Engine/Plugins/2D/Paper2D/Source/Paper2D/Private/PaperRenderSceneProxy.h
--- a/Engine/Plugins/2D/Paper2D/Source/Paper2D/Private/PaperRenderSceneProxy.h
+++ b/Engine/Plugins/2D/Paper2D/Source/Paper2D/Private/PaperRenderSceneProxy.h
@@ -94,7 +94,7 @@ struct PAPER2D_API FSpriteRenderSection
 			const FVector4& SourceVert = Record.RenderVerts[VertexIndex];
 			
 			const FVector Pos((PaperAxisX * SourceVert.X) + (PaperAxisY * SourceVert.Y) + Record.Destination);
-			const FVector2D UV(SourceVert.Z, SourceVert.W);
+			const FVector2f UV(SourceVert.Z, SourceVert.W);	// LWC_TODO: Precision loss
 
 			new (Vertices) FDynamicMeshVertex(Pos, FPaperSpriteTangents::PackedNormalX.ToFVector(), FPaperSpriteTangents::PackedNormalZ.ToFVector(), UV, VertColor);
 		}
@@ -105,7 +105,7 @@ struct PAPER2D_API FSpriteRenderSection
 	{
 		const FVector Pos((PaperAxisX * X) + (PaperAxisY * Y) + Origin);
 
-		new (Vertices) FDynamicMeshVertex(Pos, FPaperSpriteTangents::PackedNormalX.ToFVector(), FPaperSpriteTangents::PackedNormalZ.ToFVector(), FVector2D(U, V), Color);
+		new (Vertices) FDynamicMeshVertex(Pos, FPaperSpriteTangents::PackedNormalX.ToFVector(), FPaperSpriteTangents::PackedNormalZ.ToFVector(), FVector2f(U, V), Color);
 		++NumVertices;
 	}
 
diff --git a/Engine/Plugins/2D/Paper2D/Source/Paper2DEditor/Private/PaperEditorShared/SpriteGeometryEditing.cpp b/Engine/Plugins/2D/Paper2D/Source/Paper2DEditor/Private/PaperEditorShared/SpriteGeometryEditing.cpp
--- a/Engine/Plugins/2D/Paper2D/Source/Paper2DEditor/Private/PaperEditorShared/SpriteGeometryEditing.cpp
+++ b/Engine/Plugins/2D/Paper2D/Source/Paper2DEditor/Private/PaperEditorShared/SpriteGeometryEditing.cpp
@@ -386,7 +386,7 @@ void FSpriteGeometryEditingHelper::DrawGeometry(const FSceneView& View, FPrimiti
 
 				FDynamicMeshVertex MeshVertex;
 				MeshVertex.Color = BackgroundColor;
-				MeshVertex.TextureCoordinate[0] = FVector2D::ZeroVector;
+				MeshVertex.TextureCoordinate[0] = FVector2f::ZeroVector;
 				MeshVertex.SetTangents(PaperAxisX, PaperAxisY, PaperAxisZ);
 
 				for (const FVector2D& SrcTriangleVertex : TriangulatedPolygonVertices)
diff --git a/Engine/Plugins/Animation/ControlRig/Source/ControlRig/Private/Sequencer/MovieSceneControlRigParameterSection.cpp b/Engine/Plugins/Animation/ControlRig/Source/ControlRig/Private/Sequencer/MovieSceneControlRigParameterSection.cpp
--- a/Engine/Plugins/Animation/ControlRig/Source/ControlRig/Private/Sequencer/MovieSceneControlRigParameterSection.cpp
+++ b/Engine/Plugins/Animation/ControlRig/Source/ControlRig/Private/Sequencer/MovieSceneControlRigParameterSection.cpp
@@ -2971,7 +2971,7 @@ TOptional<FVector2D> UMovieSceneControlRigParameterSection::EvaluateVector2DPara
 		FVector2f Value(0.0f, 0.0f);
 		FloatChannels[ChannelInfo->ChannelIndex]->Evaluate(InTime, Value.X);
 		FloatChannels[ChannelInfo->ChannelIndex + 1]->Evaluate(InTime, Value.Y);
-		OptValue = Value;
+		OptValue = FVector2D(Value);
 	}
 	return OptValue;
 }
diff --git a/Engine/Plugins/Animation/ControlRig/Source/ControlRig/Private/Sequencer/MovieSceneControlRigParameterTemplate.cpp b/Engine/Plugins/Animation/ControlRig/Source/ControlRig/Private/Sequencer/MovieSceneControlRigParameterTemplate.cpp
--- a/Engine/Plugins/Animation/ControlRig/Source/ControlRig/Private/Sequencer/MovieSceneControlRigParameterTemplate.cpp
+++ b/Engine/Plugins/Animation/ControlRig/Source/ControlRig/Private/Sequencer/MovieSceneControlRigParameterTemplate.cpp
@@ -1768,7 +1768,7 @@ void FMovieSceneControlRigParameterTemplate::EvaluateCurvesWithMasks(const FMovi
 				}
 			}
 
-			Values.Vector2DValues.Emplace(Vector2D.ParameterName, Value);
+			Values.Vector2DValues.Emplace(Vector2D.ParameterName, FVector2D(Value));
 		}
 
 		for (int32 Index = 0; Index < Vectors.Num(); ++Index)
diff --git a/Engine/Plugins/Editor/GLTFImporter/Source/GLTFCore/Private/GLTF/GLTFAccessor.cpp b/Engine/Plugins/Editor/GLTFImporter/Source/GLTFCore/Private/GLTF/GLTFAccessor.cpp
--- a/Engine/Plugins/Editor/GLTFImporter/Source/GLTFCore/Private/GLTF/GLTFAccessor.cpp
+++ b/Engine/Plugins/Editor/GLTFImporter/Source/GLTFCore/Private/GLTF/GLTFAccessor.cpp
@@ -350,7 +350,7 @@ namespace GLTF
 				}
 				else if (Normalized)
 				{
-					return GetNormalized<FVector2f, 2>(ComponentType, Pointer);
+					return FVector2D(GetNormalized<FVector2f, 2>(ComponentType, Pointer));
 				}
 			}
 		}
@@ -380,7 +380,7 @@ namespace GLTF
 				}
 				else if (Normalized)
 				{
-					return GetNormalized<FVector3f, 3>(ComponentType, Pointer);
+					return FVector(GetNormalized<FVector3f, 3>(ComponentType, Pointer));
 				}
 			}
 		}
diff --git a/Engine/Plugins/Editor/GeometryMode/Source/GeometryMode/Private/GeometryEdMode.cpp b/Engine/Plugins/Editor/GeometryMode/Source/GeometryMode/Private/GeometryEdMode.cpp
--- a/Engine/Plugins/Editor/GeometryMode/Source/GeometryMode/Private/GeometryEdMode.cpp
+++ b/Engine/Plugins/Editor/GeometryMode/Source/GeometryMode/Private/GeometryEdMode.cpp
@@ -583,12 +583,12 @@ void FEdModeGeometry::RenderPoly( const FSceneView* View, FViewport* Viewport, F
 				}
 
 				// Draw Polygon Triangles
-				const int32 VertexOffset = MeshBuilder.AddVertex(Verts[0], FVector2D::ZeroVector, FVector(1,0,0), FVector(0,1,0), FVector(0,0,1), FColor::White);
-				MeshBuilder.AddVertex(Verts[1], FVector2D::ZeroVector, FVector(1,0,0), FVector(0,1,0), FVector(0,0,1), FColor::White);
+				const int32 VertexOffset = MeshBuilder.AddVertex(Verts[0], FVector2f::ZeroVector, FVector(1,0,0), FVector(0,1,0), FVector(0,0,1), FColor::White);
+				MeshBuilder.AddVertex(Verts[1], FVector2f::ZeroVector, FVector(1,0,0), FVector(0,1,0), FVector(0,0,1), FColor::White);
 
 				for( int32 VertIdx = 2 ; VertIdx < Verts.Num() ; ++VertIdx )
 				{
-					MeshBuilder.AddVertex(Verts[VertIdx], FVector2D::ZeroVector, FVector(1,0,0), FVector(0,1,0), FVector(0,0,1), FColor::White);
+					MeshBuilder.AddVertex(Verts[VertIdx], FVector2f::ZeroVector, FVector(1,0,0), FVector(0,1,0), FVector(0,0,1), FColor::White);
 					MeshBuilder.AddTriangle( VertexOffset + VertIdx - 1, VertexOffset, VertexOffset + VertIdx);
 				}
 
diff --git a/Engine/Plugins/Editor/SpeedTreeImporter/Source/SpeedTreeImporter/Private/SpeedTreeImportFactory.cpp b/Engine/Plugins/Editor/SpeedTreeImporter/Source/SpeedTreeImporter/Private/SpeedTreeImportFactory.cpp
--- a/Engine/Plugins/Editor/SpeedTreeImporter/Source/SpeedTreeImporter/Private/SpeedTreeImportFactory.cpp
+++ b/Engine/Plugins/Editor/SpeedTreeImporter/Source/SpeedTreeImporter/Private/SpeedTreeImportFactory.cpp
@@ -1846,7 +1846,7 @@ FVertexInstanceID ProcessTriangleCorner(
 	// keep texcoords padded to align indices
 	for( int32 PadIndex = 0; PadIndex < NumUVs; ++PadIndex )
 	{
-		VertexInstanceUVs.Set(VertexInstanceID, PadIndex, FVector2D(0.0f, 0.0f));
+		VertexInstanceUVs.Set(VertexInstanceID, PadIndex, FVector2f::ZeroVector);
 	}
 
 	// All texcoords are packed into 4 float4 vertex attributes
@@ -1866,28 +1866,28 @@ FVertexInstanceID ProcessTriangleCorner(
 
 	// diffuse
 	DrawCall->GetProperty( SpeedTree::VERTEX_PROPERTY_DIFFUSE_TEXCOORDS, VertexIndex, Data );
-	VertexInstanceUVs.Set( VertexInstanceID, 0, FVector2D( Data[ 0 ], Data[ 1 ] ) );
+	VertexInstanceUVs.Set( VertexInstanceID, 0, FVector2f( Data[ 0 ], Data[ 1 ] ) );
 
 	// lightmap
 	DrawCall->GetProperty( SpeedTree::VERTEX_PROPERTY_LIGHTMAP_TEXCOORDS, VertexIndex, Data );
-	VertexInstanceUVs.Set( VertexInstanceID, 1, FVector2D( Data[ 0 ], Data[ 1 ] ) );
+	VertexInstanceUVs.Set( VertexInstanceID, 1, FVector2f( Data[ 0 ], Data[ 1 ] ) );
 
 	// branch wind
 	DrawCall->GetProperty( SpeedTree::VERTEX_PROPERTY_WIND_BRANCH_DATA, VertexIndex, Data );
-	VertexInstanceUVs.Set( VertexInstanceID, 2, FVector2D( Data[ 0 ], Data[ 1 ] ) );
+	VertexInstanceUVs.Set( VertexInstanceID, 2, FVector2f( Data[ 0 ], Data[ 1 ] ) );
 
 	// lod
 	if( RenderState->m_bFacingLeavesPresent )
 	{
 		DrawCall->GetProperty( SpeedTree::VERTEX_PROPERTY_LEAF_CARD_LOD_SCALAR, VertexIndex, Data );
-		VertexInstanceUVs.Set( VertexInstanceID, 3, FVector2D( Data[ 0 ], 0.0f ) );
-		VertexInstanceUVs.Set( VertexInstanceID, 4, FVector2D( 0.0f, 0.0f ) );
+		VertexInstanceUVs.Set( VertexInstanceID, 3, FVector2f( Data[ 0 ], 0.0f ) );
+		VertexInstanceUVs.Set( VertexInstanceID, 4, FVector2f( 0.0f, 0.0f ) );
 	}
 	else
 	{
 		DrawCall->GetProperty( SpeedTree::VERTEX_PROPERTY_LOD_POSITION, VertexIndex, Data );
-		VertexInstanceUVs.Set( VertexInstanceID, 3, FVector2D( -Data[ 0 ], Data[ 1 ] ) );
-		VertexInstanceUVs.Set( VertexInstanceID, 4, FVector2D( Data[ 2 ], 0.0f ) );
+		VertexInstanceUVs.Set( VertexInstanceID, 3, FVector2f( -Data[ 0 ], Data[ 1 ] ) );
+		VertexInstanceUVs.Set( VertexInstanceID, 4, FVector2f( Data[ 2 ], 0.0f ) );
 	}
 
 	// other
@@ -1895,19 +1895,19 @@ FVertexInstanceID ProcessTriangleCorner(
 	{
 		// detail
 		DrawCall->GetProperty( SpeedTree::VERTEX_PROPERTY_DETAIL_TEXCOORDS, VertexIndex, Data );
-		VertexInstanceUVs.Set( VertexInstanceID, 5, FVector2D( Data[ 0 ], Data[ 1 ] ) );
+		VertexInstanceUVs.Set( VertexInstanceID, 5, FVector2f( Data[ 0 ], Data[ 1 ] ) );
 
 		// branch seam
 		DrawCall->GetProperty( SpeedTree::VERTEX_PROPERTY_BRANCH_SEAM_DIFFUSE, VertexIndex, Data );
-		VertexInstanceUVs.Set( VertexInstanceID, 6, FVector2D( Data[ 0 ], Data[ 1 ] ) );
-		VertexInstanceUVs.Set( VertexInstanceID, 4, FVector2D( VertexInstanceUVs.Get( VertexInstanceID, 4 ).X, Data[ 2 ] ) );
+		VertexInstanceUVs.Set( VertexInstanceID, 6, FVector2f( Data[ 0 ], Data[ 1 ] ) );
+		VertexInstanceUVs.Set( VertexInstanceID, 4, FVector2f( VertexInstanceUVs.Get( VertexInstanceID, 4 ).X, Data[ 2 ] ) );
 	}
 	else if( RenderState->m_bFrondsPresent )
 	{
 		// frond wind
 		DrawCall->GetProperty( SpeedTree::VERTEX_PROPERTY_WIND_EXTRA_DATA, VertexIndex, Data );
-		VertexInstanceUVs.Set( VertexInstanceID, 5, FVector2D( Data[ 0 ], Data[ 1 ] ) );
-		VertexInstanceUVs.Set( VertexInstanceID, 6, FVector2D( Data[ 2 ], 0.0f ) );
+		VertexInstanceUVs.Set( VertexInstanceID, 5, FVector2f( Data[ 0 ], Data[ 1 ] ) );
+		VertexInstanceUVs.Set( VertexInstanceID, 6, FVector2f( Data[ 2 ], 0.0f ) );
 	}
 	else if( RenderState->m_bLeavesPresent || RenderState->m_bFacingLeavesPresent )
 	{
@@ -1922,15 +1922,15 @@ FVertexInstanceID ProcessTriangleCorner(
 		{
 			DrawCall->GetProperty( SpeedTree::VERTEX_PROPERTY_LEAF_ANCHOR_POINT, VertexIndex, Data );
 		}
-		VertexInstanceUVs.Set( VertexInstanceID, 4, FVector2D( VertexInstanceUVs.Get( VertexInstanceID, 4 ).X, -Data[ 0 ] ) );
-		VertexInstanceUVs.Set( VertexInstanceID, 5, FVector2D( Data[ 1 ], Data[ 2 ] ) );
+		VertexInstanceUVs.Set( VertexInstanceID, 4, FVector2f( VertexInstanceUVs.Get( VertexInstanceID, 4 ).X, -Data[ 0 ] ) );
+		VertexInstanceUVs.Set( VertexInstanceID, 5, FVector2f( Data[ 1 ], Data[ 2 ] ) );
 
 		// leaf wind
 		DrawCall->GetProperty( SpeedTree::VERTEX_PROPERTY_WIND_EXTRA_DATA, VertexIndex, Data );
-		VertexInstanceUVs.Set( VertexInstanceID, 6, FVector2D( Data[ 0 ], Data[ 1 ] ) );
-		VertexInstanceUVs.Set( VertexInstanceID, 7, FVector2D( Data[ 2 ], 0 ) );
+		VertexInstanceUVs.Set( VertexInstanceID, 6, FVector2f( Data[ 0 ], Data[ 1 ] ) );
+		VertexInstanceUVs.Set( VertexInstanceID, 7, FVector2f( Data[ 2 ], 0 ) );
 		DrawCall->GetProperty( SpeedTree::VERTEX_PROPERTY_WIND_FLAGS, VertexIndex, Data );
-		VertexInstanceUVs.Set( VertexInstanceID, 7, FVector2D( VertexInstanceUVs.Get( VertexInstanceID, 7 ).X, Data[ 0 ] ) );
+		VertexInstanceUVs.Set( VertexInstanceID, 7, FVector2f( VertexInstanceUVs.Get( VertexInstanceID, 7 ).X, Data[ 0 ] ) );
 	}
 	return VertexInstanceID;
 }
@@ -2379,11 +2379,11 @@ UObject* USpeedTreeImportFactory::FactoryCreateBinary7(UClass* InClass, UObject*
 								VertexInstanceBinormalSigns[VertexInstanceID] = GetBasisDeterminantSign(TangentX.GetSafeNormal(), TangentY.GetSafeNormal(), TangentZ.GetSafeNormal());
 								if (bRotated)
 								{
-									VertexInstanceUVs.Set(VertexInstanceID, 0, FVector2D(TexCoords[0] + Vertex[1] * TexCoords[2], TexCoords[1] + Vertex[0] * TexCoords[3]));
+									VertexInstanceUVs.Set(VertexInstanceID, 0, FVector2f(TexCoords[0] + Vertex[1] * TexCoords[2], TexCoords[1] + Vertex[0] * TexCoords[3]));
 								}
 								else
 								{
-									VertexInstanceUVs.Set(VertexInstanceID, 0, FVector2D(TexCoords[0] + Vertex[0] * TexCoords[2], TexCoords[1] + Vertex[1] * TexCoords[3]));
+									VertexInstanceUVs.Set(VertexInstanceID, 0, FVector2f(TexCoords[0] + Vertex[0] * TexCoords[2], TexCoords[1] + Vertex[1] * TexCoords[3]));
 								}
 
 								// lightmap coord
@@ -2724,57 +2724,57 @@ UObject* USpeedTreeImportFactory::FactoryCreateBinary8(UClass* InClass, UObject*
 						// 7	0					0					Leaf Wind Z, Leaf Group
 
 						// diffuse
-						VertexInstanceUVs.Set(VertexInstanceID, 0, FVector2D(Vertex.m_vTexCoord.x, Vertex.m_vTexCoord.y));
+						VertexInstanceUVs.Set(VertexInstanceID, 0, FVector2f(Vertex.m_vTexCoord.x, Vertex.m_vTexCoord.y));
 
 						// lightmap
-						VertexInstanceUVs.Set(VertexInstanceID, 1, FVector2D(Vertex.m_vLightmapTexCoord.x, Vertex.m_vLightmapTexCoord.y));
+						VertexInstanceUVs.Set(VertexInstanceID, 1, FVector2f(Vertex.m_vLightmapTexCoord.x, Vertex.m_vLightmapTexCoord.y));
 
 						if (DrawCall.m_eWindGeometryType == GameEngine8::Billboard)
 						{
-							VertexInstanceUVs.Set(VertexInstanceID, 2, FVector2D((Vertex.m_vNormal.z > 0.5f) ? 1.0f : 0.0f, 0.0f));
+							VertexInstanceUVs.Set(VertexInstanceID, 2, FVector2f((Vertex.m_vNormal.z > 0.5f) ? 1.0f : 0.0f, 0.0f));
 						}
 						else
 						{
 							// branch wind
-							VertexInstanceUVs.Set(VertexInstanceID, 2, FVector2D(Vertex.m_vWindBranch.x, Vertex.m_vWindBranch.y));
+							VertexInstanceUVs.Set(VertexInstanceID, 2, FVector2f(Vertex.m_vWindBranch.x, Vertex.m_vWindBranch.y));
 
 							// lod
 							FVector vLodPosition = FVector(Vertex.m_vAnchor.x, Vertex.m_vAnchor.y, Vertex.m_vAnchor.z) +
 								FVector(Vertex.m_vLodOffset.x, Vertex.m_vLodOffset.y, Vertex.m_vLodOffset.z);
-							VertexInstanceUVs.Set(VertexInstanceID, 3, FVector2D(vLodPosition[0], vLodPosition[1]));
-							VertexInstanceUVs.Set(VertexInstanceID, 4, FVector2D(vLodPosition[2], 0.0f));
+							VertexInstanceUVs.Set(VertexInstanceID, 3, FVector2f(vLodPosition[0], vLodPosition[1]));
+							VertexInstanceUVs.Set(VertexInstanceID, 4, FVector2f(vLodPosition[2], 0.0f));
 
 							// other
 							if (DrawCall.m_eWindGeometryType == GameEngine8::Branch)
 							{
 								// detail (not used in v8)
-								VertexInstanceUVs.Set(VertexInstanceID, 5, FVector2D(0.0f, 0.0f));
+								VertexInstanceUVs.Set(VertexInstanceID, 5, FVector2f(0.0f, 0.0f));
 
 								// branch seam
-								VertexInstanceUVs.Set(VertexInstanceID, 6, FVector2D(0.0f, 0.0f));
-								VertexInstanceUVs.Set(VertexInstanceID, 4, FVector2D(VertexInstanceUVs.Get(VertexInstanceID, 4).X, Vertex.m_fBlendWeight));
+								VertexInstanceUVs.Set(VertexInstanceID, 6, FVector2f(0.0f, 0.0f));
+								VertexInstanceUVs.Set(VertexInstanceID, 4, FVector2f(VertexInstanceUVs.Get(VertexInstanceID, 4).X, Vertex.m_fBlendWeight));
 
 								// keep alignment
-								VertexInstanceUVs.Set(VertexInstanceID, 7, FVector2D(0.0f, 0.0f));
+								VertexInstanceUVs.Set(VertexInstanceID, 7, FVector2f(0.0f, 0.0f));
 							}
 							else if (DrawCall.m_eWindGeometryType == GameEngine8::Frond)
 							{
 								// frond wind
-								VertexInstanceUVs.Set(VertexInstanceID, 5, FVector2D(Vertex.m_vWindNonBranch.x, Vertex.m_vWindNonBranch.y));
-								VertexInstanceUVs.Set(VertexInstanceID, 6, FVector2D(Vertex.m_vWindNonBranch.z, 0.0f));
+								VertexInstanceUVs.Set(VertexInstanceID, 5, FVector2f(Vertex.m_vWindNonBranch.x, Vertex.m_vWindNonBranch.y));
+								VertexInstanceUVs.Set(VertexInstanceID, 6, FVector2f(Vertex.m_vWindNonBranch.z, 0.0f));
 
 								// keep alignment
-								VertexInstanceUVs.Set(VertexInstanceID, 7, FVector2D(0.0f, 0.0f));
+								VertexInstanceUVs.Set(VertexInstanceID, 7, FVector2f(0.0f, 0.0f));
 							}
 							else if (DrawCall.m_eWindGeometryType == GameEngine8::Leaf || DrawCall.m_eWindGeometryType == GameEngine8::FacingLeaf)
 							{
 								// anchor
-								VertexInstanceUVs.Set(VertexInstanceID, 4, FVector2D(VertexInstanceUVs.Get(VertexInstanceID, 4).X, Vertex.m_vAnchor.x));
-								VertexInstanceUVs.Set(VertexInstanceID, 5, FVector2D(Vertex.m_vAnchor.y, Vertex.m_vAnchor.z));
+								VertexInstanceUVs.Set(VertexInstanceID, 4, FVector2f(VertexInstanceUVs.Get(VertexInstanceID, 4).X, Vertex.m_vAnchor.x));
+								VertexInstanceUVs.Set(VertexInstanceID, 5, FVector2f(Vertex.m_vAnchor.y, Vertex.m_vAnchor.z));
 
 								// leaf wind
-								VertexInstanceUVs.Set(VertexInstanceID, 6, FVector2D(Vertex.m_vWindNonBranch.x, Vertex.m_vWindNonBranch.y));
-								VertexInstanceUVs.Set(VertexInstanceID, 7, FVector2D(Vertex.m_vWindNonBranch.z, (Vertex.m_bWindLeaf2Flag ? 1.0f : 0.0f)));
+								VertexInstanceUVs.Set(VertexInstanceID, 6, FVector2f(Vertex.m_vWindNonBranch.x, Vertex.m_vWindNonBranch.y));
+								VertexInstanceUVs.Set(VertexInstanceID, 7, FVector2f(Vertex.m_vWindNonBranch.z, (Vertex.m_bWindLeaf2Flag ? 1.0f : 0.0f)));
 							}
 						}
 						CornerVertexInstanceIDs[Corner] = VertexInstanceID;
@@ -3058,27 +3058,27 @@ UObject* USpeedTreeImportFactory::FactoryCreateBinary9(UClass* InClass, UObject*
 
 						// diffuse
 						int32 CurrentUV = 0;
-						VertexInstanceUVs.Set(VertexInstanceID, CurrentUV++, FVector2D(Vertex.m_vTexCoord.x, Vertex.m_vTexCoord.y));
+						VertexInstanceUVs.Set(VertexInstanceID, CurrentUV++, FVector2f(Vertex.m_vTexCoord.x, Vertex.m_vTexCoord.y));
 
 						// branch1 / ripple
-						VertexInstanceUVs.Set(VertexInstanceID, CurrentUV++, FVector2D(Vertex.m_vBranchWind1.x, Vertex.m_vBranchWind1.y));
-						VertexInstanceUVs.Set(VertexInstanceID, CurrentUV++, FVector2D(Vertex.m_vBranchWind1.z, Vertex.m_fRippleWeight));
+						VertexInstanceUVs.Set(VertexInstanceID, CurrentUV++, FVector2f(Vertex.m_vBranchWind1.x, Vertex.m_vBranchWind1.y));
+						VertexInstanceUVs.Set(VertexInstanceID, CurrentUV++, FVector2f(Vertex.m_vBranchWind1.z, Vertex.m_fRippleWeight));
 
 						// lightmap (lightmass can only access 4 uvs)
-						VertexInstanceUVs.Set(VertexInstanceID, CurrentUV++, FVector2D(Vertex.m_vLightmapTexCoord.x, Vertex.m_vLightmapTexCoord.y));
+						VertexInstanceUVs.Set(VertexInstanceID, CurrentUV++, FVector2f(Vertex.m_vLightmapTexCoord.x, Vertex.m_vLightmapTexCoord.y));
 
 						// branch 2
 						if (bHasBranch2Data)
 						{
-							VertexInstanceUVs.Set(VertexInstanceID, CurrentUV++, FVector2D(Vertex.m_vBranchWind2.x, Vertex.m_vBranchWind2.y));
-							VertexInstanceUVs.Set(VertexInstanceID, CurrentUV++, FVector2D(Vertex.m_vBranchWind2.z, 0.0f));
+							VertexInstanceUVs.Set(VertexInstanceID, CurrentUV++, FVector2f(Vertex.m_vBranchWind2.x, Vertex.m_vBranchWind2.y));
+							VertexInstanceUVs.Set(VertexInstanceID, CurrentUV++, FVector2f(Vertex.m_vBranchWind2.z, 0.0f));
 						}
 
 						// camera-facing
 						if (bHasFacingData)
 						{
-							VertexInstanceUVs.Set(VertexInstanceID, CurrentUV++, FVector2D(Vertex.m_vAnchor.x, Vertex.m_vAnchor.y));
-							VertexInstanceUVs.Set(VertexInstanceID, CurrentUV++, FVector2D(Vertex.m_vAnchor.z, Vertex.m_bCameraFacing ? 1.0f : 0.0f));
+							VertexInstanceUVs.Set(VertexInstanceID, CurrentUV++, FVector2f(Vertex.m_vAnchor.x, Vertex.m_vAnchor.y));
+							VertexInstanceUVs.Set(VertexInstanceID, CurrentUV++, FVector2f(Vertex.m_vAnchor.z, Vertex.m_bCameraFacing ? 1.0f : 0.0f));
 						}
 					}
 
diff --git a/Engine/Plugins/Enterprise/DataprepEditor/Source/DataprepEditor/Private/Widgets/DataprepGraph/SDataprepGraphTrackNode.cpp b/Engine/Plugins/Enterprise/DataprepEditor/Source/DataprepEditor/Private/Widgets/DataprepGraph/SDataprepGraphTrackNode.cpp
--- a/Engine/Plugins/Enterprise/DataprepEditor/Source/DataprepEditor/Private/Widgets/DataprepGraph/SDataprepGraphTrackNode.cpp
+++ b/Engine/Plugins/Enterprise/DataprepEditor/Source/DataprepEditor/Private/Widgets/DataprepGraph/SDataprepGraphTrackNode.cpp
@@ -716,7 +716,7 @@ void SDataprepGraphTrackNode::OnStartNodeDrag(const TSharedRef<SDataprepGraphBas
 
 	const FGeometry& PanelGeometry = GetOwnerPanel()->GetTickSpaceGeometry();
 
-	const FVector2D PanelPosition = PanelGeometry.AbsolutePosition;
+	const FVector2f PanelPosition = PanelGeometry.AbsolutePosition;
 	const FVector2D PanelSize = PanelGeometry.GetLocalSize();
 	RECT Boundaries;
 	Boundaries.left = FMath::RoundToInt(PanelPosition.X);
diff --git a/Engine/Plugins/Enterprise/DatasmithC4DImporter/Source/DatasmithC4DTranslator/Private/DatasmithC4DImporter.cpp b/Engine/Plugins/Enterprise/DatasmithC4DImporter/Source/DatasmithC4DTranslator/Private/DatasmithC4DImporter.cpp
--- a/Engine/Plugins/Enterprise/DatasmithC4DImporter/Source/DatasmithC4DTranslator/Private/DatasmithC4DImporter.cpp
+++ b/Engine/Plugins/Enterprise/DatasmithC4DImporter/Source/DatasmithC4DTranslator/Private/DatasmithC4DImporter.cpp
@@ -2861,7 +2861,7 @@ TSharedPtr<IDatasmithMeshElement> FDatasmithC4DImporter::ImportMesh(melange::Pol
 	IDsCopy.SetNumUninitialized(3);
 	TArray<FVector> QuadNormals;
 	QuadNormals.SetNumZeroed(4);
-	TArray<FVector2D> QuadUVs;
+	TArray<FVector2f> QuadUVs;
 	QuadUVs.SetNumZeroed(4);
 
 	// Used to check for degenerate triangles
@@ -2947,7 +2947,7 @@ TSharedPtr<IDatasmithMeshElement> FDatasmithC4DImporter::ImportMesh(melange::Pol
 			for (int32 VertexIndex = 0; VertexIndex < 4; ++VertexIndex)
 			{
 				melange::Vector& PointUVs = UVs[VertexIndex];
-				FVector2D& UnrealUVs = QuadUVs[VertexIndex];
+				FVector2f& UnrealUVs = QuadUVs[VertexIndex];
 
 				if (PointUVs.z != 0.0f && PointUVs.z != 1.0f)
 				{
diff --git a/Engine/Plugins/Enterprise/DatasmithCADImporter/Source/CADLibrary/Private/CADKernelTools.cpp b/Engine/Plugins/Enterprise/DatasmithCADImporter/Source/CADLibrary/Private/CADKernelTools.cpp
--- a/Engine/Plugins/Enterprise/DatasmithCADImporter/Source/CADLibrary/Private/CADKernelTools.cpp
+++ b/Engine/Plugins/Enterprise/DatasmithCADImporter/Source/CADLibrary/Private/CADKernelTools.cpp
@@ -212,7 +212,7 @@ namespace CADLibrary
 					for (int32 Index = 0; Index < TriangleCount; Index++)
 					{
 						const FVertexInstanceID VertexInstanceID = MeshVertexInstanceIDs[IndexFace + Index];
-						VertexInstanceUVs.Set(VertexInstanceID, UVChannel, FaceMesh->UVMap[TriangleVertexIndices[IndexFace + Orientation[Index]]]);
+						VertexInstanceUVs.Set(VertexInstanceID, UVChannel, FVector2f(FaceMesh->UVMap[TriangleVertexIndices[IndexFace + Orientation[Index]]]));
 
 						VertexInstanceColors[VertexInstanceID] = FLinearColor::White;
 						VertexInstanceTangents[VertexInstanceID] = FVector(ForceInitToZero);
diff --git a/Engine/Plugins/Enterprise/DatasmithCADImporter/Source/CADLibrary/Private/MeshDescriptionHelper.cpp b/Engine/Plugins/Enterprise/DatasmithCADImporter/Source/CADLibrary/Private/MeshDescriptionHelper.cpp
--- a/Engine/Plugins/Enterprise/DatasmithCADImporter/Source/CADLibrary/Private/MeshDescriptionHelper.cpp
+++ b/Engine/Plugins/Enterprise/DatasmithCADImporter/Source/CADLibrary/Private/MeshDescriptionHelper.cpp
@@ -476,7 +476,7 @@ namespace CADLibrary
 						for (int32 VertexIndex = 0; VertexIndex < VertexCountPerFace; VertexIndex++)
 						{
 							const FVertexInstanceID& VertexInstanceID = MeshVertexInstanceIDs[FaceIndex + Orientation[VertexIndex]];
-							VertexInstanceUVs.Set(VertexInstanceID, UVChannel, Tessellation.TexCoordArray[NewFaceIndex[FaceIndex + VertexIndex]]);
+							VertexInstanceUVs.Set(VertexInstanceID, UVChannel, FVector2f(Tessellation.TexCoordArray[NewFaceIndex[FaceIndex + VertexIndex]]));
 						}
 					}
 				}
diff --git a/Engine/Plugins/Enterprise/DatasmithCADImporter/Source/DatasmithOpenNurbsTranslator/Private/DatasmithOpenNurbsTranslator.cpp b/Engine/Plugins/Enterprise/DatasmithCADImporter/Source/DatasmithOpenNurbsTranslator/Private/DatasmithOpenNurbsTranslator.cpp
--- a/Engine/Plugins/Enterprise/DatasmithCADImporter/Source/DatasmithOpenNurbsTranslator/Private/DatasmithOpenNurbsTranslator.cpp
+++ b/Engine/Plugins/Enterprise/DatasmithCADImporter/Source/DatasmithOpenNurbsTranslator/Private/DatasmithOpenNurbsTranslator.cpp
@@ -517,7 +517,7 @@ namespace DatasmithOpenNurbsTranslatorUtils
 						if (!UVValues.ContainsNaN())
 						{
 							// Convert UVs - bottom-left Rhino's texture origin to Unreal's top-left
-							VertexInstanceUVs.Set(CornerVertexInstanceIDs[CornerIndex], 0, FVector2D(UVValues[0], 1.f - UVValues[1]));
+							VertexInstanceUVs.Set(CornerVertexInstanceIDs[CornerIndex], 0, FVector2f(UVValues[0], 1.f - UVValues[1]));
 						}
 					}
 				}
diff --git a/Engine/Plugins/Enterprise/DatasmithCADImporter/Source/DatasmithWireTranslator/Private/OpenModelUtils.cpp b/Engine/Plugins/Enterprise/DatasmithCADImporter/Source/DatasmithWireTranslator/Private/OpenModelUtils.cpp
--- a/Engine/Plugins/Enterprise/DatasmithCADImporter/Source/DatasmithWireTranslator/Private/OpenModelUtils.cpp
+++ b/Engine/Plugins/Enterprise/DatasmithCADImporter/Source/DatasmithWireTranslator/Private/OpenModelUtils.cpp
@@ -1748,7 +1748,7 @@ bool OpenModelUtils::TransferAlMeshToMeshDescription(const AlMesh& AliasMesh, FM
 					{
 						FVector2D UVValues(AlUVs[2 * Triangles[TIndex] + 0], AlUVs[2 * Triangles[TIndex] + 1]);
 						UVBBox += FVector(UVValues, 0.0f);
-						VertexInstanceUVs.Set(CornerVertexInstanceIDs[VertexIndex], 0, UVValues);
+						VertexInstanceUVs.Set(CornerVertexInstanceIDs[VertexIndex], 0, FVector2f(UVValues));
 					}
 				}
 
@@ -1790,7 +1790,7 @@ bool OpenModelUtils::TransferAlMeshToMeshDescription(const AlMesh& AliasMesh, FM
 					{
 						FVector2D UVValues(AlUVs[2 * Triangles[VertexIndex] + 0], AlUVs[2 * Triangles[VertexIndex] + 1]);
 						UVBBox += FVector(UVValues, 0.0f);
-						VertexInstanceUVs.Set(CornerVertexInstanceIDs[VertexIndex], 0, UVValues);
+						VertexInstanceUVs.Set(CornerVertexInstanceIDs[VertexIndex], 0, FVector2f(UVValues));
 					}
 				}
 
diff --git a/Engine/Plugins/Enterprise/DatasmithFBXImporter/Source/DatasmithFBXTranslator/Private/DatasmithFBXFileImporter.cpp b/Engine/Plugins/Enterprise/DatasmithFBXImporter/Source/DatasmithFBXTranslator/Private/DatasmithFBXFileImporter.cpp
--- a/Engine/Plugins/Enterprise/DatasmithFBXImporter/Source/DatasmithFBXTranslator/Private/DatasmithFBXFileImporter.cpp
+++ b/Engine/Plugins/Enterprise/DatasmithFBXImporter/Source/DatasmithFBXTranslator/Private/DatasmithFBXFileImporter.cpp
@@ -1243,7 +1243,7 @@ void FDatasmithFBXFileImporter::DoImportMesh(FbxMesh* InMesh, FDatasmithFBXScene
 				int32 UVMapIndex = (FBXUVs[UVLayerIndex].UVMappingMode == FbxLayerElement::eByControlPoint) ? CornerVertexIDs[CornerIndex].GetValue() : FirstFbxInstanceIndexForPoly + CornerIndex;
 				int32 UVIndex = (FBXUVs[UVLayerIndex].UVReferenceMode == FbxLayerElement::eDirect) ? UVMapIndex : FBXUVs[UVLayerIndex].LayerElementUV->GetIndexArray().GetAt(UVMapIndex);
 				FbxVector2 UVVector = FBXUVs[UVLayerIndex].LayerElementUV->GetDirectArray().GetAt(UVIndex);
-				FVector2D FinalUVVector(static_cast<float>(UVVector[0]), 1.f - static_cast<float>(UVVector[1])); // flip the Y of UVs for DirectX
+				FVector2f FinalUVVector(static_cast<float>(UVVector[0]), 1.f - static_cast<float>(UVVector[1])); // flip the Y of UVs for DirectX
 				if (!FinalUVVector.ContainsNaN())
 				{
 					VertexInstanceUVs.Set(CornerVertexInstanceIDs[CornerIndex], UVLayerIndex, FinalUVVector);
diff --git a/Engine/Plugins/Enterprise/DatasmithIFCImporter/Source/DatasmithIFCTranslator/Private/IFC/IFCStaticMeshFactory.cpp b/Engine/Plugins/Enterprise/DatasmithIFCImporter/Source/DatasmithIFCTranslator/Private/IFC/IFCStaticMeshFactory.cpp
--- a/Engine/Plugins/Enterprise/DatasmithIFCImporter/Source/DatasmithIFCTranslator/Private/IFC/IFCStaticMeshFactory.cpp
+++ b/Engine/Plugins/Enterprise/DatasmithIFCImporter/Source/DatasmithIFCTranslator/Private/IFC/IFCStaticMeshFactory.cpp
@@ -114,7 +114,7 @@ namespace IFC
 
 				for (int32 UVIndex = 0; UVIndex < NumUVs; ++UVIndex)
 				{
-					VertexInstanceUVs.Set(VertexInstanceID, UVIndex, FVector2D::ZeroVector);
+					VertexInstanceUVs.Set(VertexInstanceID, UVIndex, FVector2f::ZeroVector);
 				}
 
 				const float* Vertex = &(InObject->facesVertices[(IFCVertexIndex * (InObject->vertexElementSize / sizeof(float)))]);
diff --git a/Engine/Plugins/Enterprise/DatasmithImporter/Source/DatasmithImporter/Private/DatasmithBlueprintLibrary.cpp b/Engine/Plugins/Enterprise/DatasmithImporter/Source/DatasmithImporter/Private/DatasmithBlueprintLibrary.cpp
--- a/Engine/Plugins/Enterprise/DatasmithImporter/Source/DatasmithImporter/Private/DatasmithBlueprintLibrary.cpp
+++ b/Engine/Plugins/Enterprise/DatasmithImporter/Source/DatasmithImporter/Private/DatasmithBlueprintLibrary.cpp
@@ -1093,7 +1093,7 @@ int32 UDatasmithStaticMeshBlueprintLibrary::ComputeLightmapResolution(UStaticMes
 		{
 			uint32 VertexIndex = IndexBuffer.GetIndex(TriangleIndex * 3 + CornerIndex);
 			VertexPosition[CornerIndex] = PositionBuffer.VertexPosition(VertexIndex);
-			LightmapUVs[CornerIndex] = VertexBuffer.GetVertexUV(VertexIndex, StaticMesh->GetLightMapCoordinateIndex());
+			LightmapUVs[CornerIndex] = FVector2D(VertexBuffer.GetVertexUV(VertexIndex, StaticMesh->GetLightMapCoordinateIndex()));
 		}
 
 		const float PolygonArea = DatasmithStaticMeshBlueprintLibraryUtil::ParallelogramArea(VertexPosition[0], VertexPosition[1], VertexPosition[2]);
diff --git a/Engine/Plugins/Enterprise/DatasmithImporter/Source/DatasmithImporter/Private/DatasmithStaticMeshImporter.cpp b/Engine/Plugins/Enterprise/DatasmithImporter/Source/DatasmithImporter/Private/DatasmithStaticMeshImporter.cpp
--- a/Engine/Plugins/Enterprise/DatasmithImporter/Source/DatasmithImporter/Private/DatasmithStaticMeshImporter.cpp
+++ b/Engine/Plugins/Enterprise/DatasmithImporter/Source/DatasmithImporter/Private/DatasmithStaticMeshImporter.cpp
@@ -131,7 +131,7 @@ void FDatasmithStaticMeshImporter::CleanupMeshDescriptions(TArray<FMeshDescripti
 		{
 			for (FVertexInstanceID InstanceID : MeshDescription.VertexInstances().GetElementIDs())
 			{
-				FVector2D UVCoords = VertexInstanceUVs.Get(InstanceID, Index);
+				FVector2f UVCoords = VertexInstanceUVs.Get(InstanceID, Index);
 				if (UVCoords.ContainsNaN())
 				{
 					VertexInstanceUVs.RemoveChannel(Index);
diff --git a/Engine/Plugins/Enterprise/DatasmithImporter/Source/DatasmithImporter/Private/UVTools/UVGenerationFlattenMapping.cpp b/Engine/Plugins/Enterprise/DatasmithImporter/Source/DatasmithImporter/Private/UVTools/UVGenerationFlattenMapping.cpp
--- a/Engine/Plugins/Enterprise/DatasmithImporter/Source/DatasmithImporter/Private/UVTools/UVGenerationFlattenMapping.cpp
+++ b/Engine/Plugins/Enterprise/DatasmithImporter/Source/DatasmithImporter/Private/UVTools/UVGenerationFlattenMapping.cpp
@@ -1254,7 +1254,7 @@ void UUVGenerationFlattenMapping::GenerateUVs(FMeshDescription& InMesh, int32 UV
 			// First occurrence of vertex instance
 			if( GroupIndexPtr == nullptr )
 			{
-				UVChannels.Set( FVertexInstanceID( Face.UvIndex[Corner] ), UVChannel, Face.Uvs[Corner] );
+				UVChannels.Set( FVertexInstanceID( Face.UvIndex[Corner] ), UVChannel, FVector2f(Face.Uvs[Corner]) );
 				ProcessedVertexInstances.Add( Face.UvIndex[Corner], Face.Group );
 			}
 			// Vertex instance is shared across islands. Split it.
@@ -1276,7 +1276,7 @@ void UUVGenerationFlattenMapping::GenerateUVs(FMeshDescription& InMesh, int32 UV
 				);
 
 				// Update new vertex instance with flattened uvs
-				UVChannels.Set( NewInstanceID, UVChannel, Face.Uvs[Corner] );
+				UVChannels.Set( NewInstanceID, UVChannel, FVector2f(Face.Uvs[Corner]) );
 
 				// Update contour of associated polygon with new vertex instance
 				const FPolygonID PolygonID( Face.PolygonID );
diff --git a/Engine/Plugins/Enterprise/DatasmithImporter/Source/DatasmithImporter/Private/UVTools/UVGenerationUtils.cpp b/Engine/Plugins/Enterprise/DatasmithImporter/Source/DatasmithImporter/Private/UVTools/UVGenerationUtils.cpp
--- a/Engine/Plugins/Enterprise/DatasmithImporter/Source/DatasmithImporter/Private/UVTools/UVGenerationUtils.cpp
+++ b/Engine/Plugins/Enterprise/DatasmithImporter/Source/DatasmithImporter/Private/UVTools/UVGenerationUtils.cpp
@@ -24,7 +24,7 @@ int32 UVGenerationUtils::GetNextOpenUVChannel(UStaticMesh* StaticMesh, int32 LOD
 	for (; FirstEmptyUVs < NumberOfUVs; ++FirstEmptyUVs)
 	{
 		const TVertexInstanceAttributesConstRef<FVector2f> UVChannels = Mesh->VertexInstanceAttributes().GetAttributesRef<FVector2f>(MeshAttribute::VertexInstance::TextureCoordinate);
-		const FVector2D DefValue = UVChannels.GetDefaultValue();
+		const FVector2f DefValue = UVChannels.GetDefaultValue();
 		bool bHasNonDefaultValue = false;
 
 		for (FVertexInstanceID InstanceID : Mesh->VertexInstances().GetElementIDs())
diff --git a/Engine/Plugins/Enterprise/DatasmithImporter/Source/DatasmithTranslator/Private/Utility/DatasmithMeshHelper.cpp b/Engine/Plugins/Enterprise/DatasmithImporter/Source/DatasmithTranslator/Private/Utility/DatasmithMeshHelper.cpp
--- a/Engine/Plugins/Enterprise/DatasmithImporter/Source/DatasmithTranslator/Private/Utility/DatasmithMeshHelper.cpp
+++ b/Engine/Plugins/Enterprise/DatasmithImporter/Source/DatasmithTranslator/Private/Utility/DatasmithMeshHelper.cpp
@@ -239,7 +239,7 @@ namespace DatasmithMeshHelper
 		if (HasUVChannel(Mesh, ChannelIndex))
 		{
 			TVertexInstanceAttributesConstRef<FVector2f> UVChannels = FStaticMeshConstAttributes(Mesh).GetVertexInstanceUVs();
-			const FVector2D DefValue = UVChannels.GetDefaultValue();
+			const FVector2f DefValue = UVChannels.GetDefaultValue();
 			for (FVertexInstanceID InstanceID : Mesh.VertexInstances().GetElementIDs())
 			{
 				if (UVChannels.Get(InstanceID, ChannelIndex) != DefValue)
@@ -330,7 +330,7 @@ namespace DatasmithMeshHelper
 		{
 			if (const FVector2D* UVCoord = TexCoords.Find(VertexInstanceID))
 			{
-				UVs.Set(VertexInstanceID, 0, *UVCoord);
+				UVs.Set(VertexInstanceID, 0, FVector2f(*UVCoord));
 			}
 			else
 			{
diff --git a/Engine/Plugins/Experimental/ChaosCloth/Source/Chaos/Private/ChaosCloth/ChaosClothingSimulation.cpp b/Engine/Plugins/Experimental/ChaosCloth/Source/Chaos/Private/ChaosCloth/ChaosClothingSimulation.cpp
--- a/Engine/Plugins/Experimental/ChaosCloth/Source/Chaos/Private/ChaosCloth/ChaosClothingSimulation.cpp
+++ b/Engine/Plugins/Experimental/ChaosCloth/Source/Chaos/Private/ChaosCloth/ChaosClothingSimulation.cpp
@@ -922,9 +922,9 @@ void FClothingSimulation::DebugDrawPhysMeshShaded(FPrimitiveDrawInterface* PDI)
 			const bool bIsKinematic1 = (InvMasses[Element.Y - Offset] == 0.f);
 			const bool bIsKinematic2 = (InvMasses[Element.Z - Offset] == 0.f);
 
-			MeshBuilder.AddVertex(FDynamicMeshVertex(Pos0, Tangent, Normal, FVector2D(0.f, 0.f), bIsKinematic0 ? FColor::Purple : FColor::White));
-			MeshBuilder.AddVertex(FDynamicMeshVertex(Pos1, Tangent, Normal, FVector2D(0.f, 1.f), bIsKinematic1 ? FColor::Purple : FColor::White));
-			MeshBuilder.AddVertex(FDynamicMeshVertex(Pos2, Tangent, Normal, FVector2D(1.f, 1.f), bIsKinematic2 ? FColor::Purple : FColor::White));
+			MeshBuilder.AddVertex(FDynamicMeshVertex(Pos0, Tangent, Normal, FVector2f(0.f, 0.f), bIsKinematic0 ? FColor::Purple : FColor::White));
+			MeshBuilder.AddVertex(FDynamicMeshVertex(Pos1, Tangent, Normal, FVector2f(0.f, 1.f), bIsKinematic1 ? FColor::Purple : FColor::White));
+			MeshBuilder.AddVertex(FDynamicMeshVertex(Pos2, Tangent, Normal, FVector2f(1.f, 1.f), bIsKinematic2 ? FColor::Purple : FColor::White));
 			MeshBuilder.AddTriangle(VertexIndex, VertexIndex + 1, VertexIndex + 2);
 		}
 	}
diff --git a/Engine/Plugins/Experimental/ColorCorrectRegions/Source/ColorCorrectRegions/Private/ColorCorrectRegionsSceneViewExtension.cpp b/Engine/Plugins/Experimental/ColorCorrectRegions/Source/ColorCorrectRegions/Private/ColorCorrectRegionsSceneViewExtension.cpp
--- a/Engine/Plugins/Experimental/ColorCorrectRegions/Source/ColorCorrectRegions/Private/ColorCorrectRegionsSceneViewExtension.cpp
+++ b/Engine/Plugins/Experimental/ColorCorrectRegions/Source/ColorCorrectRegions/Private/ColorCorrectRegionsSceneViewExtension.cpp
@@ -44,32 +44,32 @@ namespace
 
 	FScreenPassTextureViewportParameters GetTextureViewportParameters(const FScreenPassTextureViewport& InViewport)
 	{
-		const FVector2D Extent(InViewport.Extent);
-		const FVector2D ViewportMin(InViewport.Rect.Min.X, InViewport.Rect.Min.Y);
-		const FVector2D ViewportMax(InViewport.Rect.Max.X, InViewport.Rect.Max.Y);
-		const FVector2D ViewportSize = ViewportMax - ViewportMin;
+		const FVector2f Extent(InViewport.Extent);	// LWC_TODO: Precision loss
+		const FVector2f ViewportMin(InViewport.Rect.Min.X, InViewport.Rect.Min.Y);
+		const FVector2f ViewportMax(InViewport.Rect.Max.X, InViewport.Rect.Max.Y);
+		const FVector2f ViewportSize = ViewportMax - ViewportMin;
 
 		FScreenPassTextureViewportParameters Parameters;
 
 		if (!InViewport.IsEmpty())
 		{
-			Parameters.Extent = Extent;
-			Parameters.ExtentInverse = FVector2D(1.0f / Extent.X, 1.0f / Extent.Y);
+			Parameters.Extent = FVector2f(Extent);
+			Parameters.ExtentInverse = FVector2f(1.0f / Extent.X, 1.0f / Extent.Y);
 
-			Parameters.ScreenPosToViewportScale = FVector2D(0.5f, -0.5f) * ViewportSize;
-			Parameters.ScreenPosToViewportBias = (0.5f * ViewportSize) + ViewportMin;
+			Parameters.ScreenPosToViewportScale = FVector2f(0.5f, -0.5f) * ViewportSize;	
+			Parameters.ScreenPosToViewportBias = (0.5f * ViewportSize) + ViewportMin;	
 
 			Parameters.ViewportMin = InViewport.Rect.Min;
 			Parameters.ViewportMax = InViewport.Rect.Max;
 
 			Parameters.ViewportSize = ViewportSize;
-			Parameters.ViewportSizeInverse = FVector2D(1.0f / Parameters.ViewportSize.X, 1.0f / Parameters.ViewportSize.Y);
+			Parameters.ViewportSizeInverse = FVector2f(1.0f / Parameters.ViewportSize.X, 1.0f / Parameters.ViewportSize.Y);
 
 			Parameters.UVViewportMin = ViewportMin * Parameters.ExtentInverse;
 			Parameters.UVViewportMax = ViewportMax * Parameters.ExtentInverse;
 
 			Parameters.UVViewportSize = Parameters.UVViewportMax - Parameters.UVViewportMin;
-			Parameters.UVViewportSizeInverse = FVector2D(1.0f / Parameters.UVViewportSize.X, 1.0f / Parameters.UVViewportSize.Y);
+			Parameters.UVViewportSizeInverse = FVector2f(1.0f / Parameters.UVViewportSize.X, 1.0f / Parameters.UVViewportSize.Y);
 
 			Parameters.UVViewportBilinearMin = Parameters.UVViewportMin + 0.5f * Parameters.ExtentInverse;
 			Parameters.UVViewportBilinearMax = Parameters.UVViewportMax - 0.5f * Parameters.ExtentInverse;
diff --git a/Engine/Plugins/Experimental/CommonConversation/Source/CommonConversationGraph/Private/SConversationGraphNode.cpp b/Engine/Plugins/Experimental/CommonConversation/Source/CommonConversationGraph/Private/SConversationGraphNode.cpp
--- a/Engine/Plugins/Experimental/CommonConversation/Source/CommonConversationGraph/Private/SConversationGraphNode.cpp
+++ b/Engine/Plugins/Experimental/CommonConversation/Source/CommonConversationGraph/Private/SConversationGraphNode.cpp
@@ -572,7 +572,7 @@ void SConversationGraphNode::UpdateGraphNode()
 void SConversationGraphNode::Tick( const FGeometry& AllottedGeometry, const double InCurrentTime, const float InDeltaTime )
 {
 	SGraphNode::Tick( AllottedGeometry, InCurrentTime, InDeltaTime );
-	CachedPosition = AllottedGeometry.AbsolutePosition / AllottedGeometry.Scale;
+	CachedPosition = FVector2D(AllottedGeometry.AbsolutePosition / AllottedGeometry.Scale);
 
 	UConversationGraphNode* MyNode = Cast<UConversationGraphNode>(GraphNode);
 
diff --git a/Engine/Plugins/Experimental/Enterprise/DatasmithRuntime/Source/Private/DatasmithRuntimeUtils.cpp b/Engine/Plugins/Experimental/Enterprise/DatasmithRuntime/Source/Private/DatasmithRuntimeUtils.cpp
--- a/Engine/Plugins/Experimental/Enterprise/DatasmithRuntime/Source/Private/DatasmithRuntimeUtils.cpp
+++ b/Engine/Plugins/Experimental/Enterprise/DatasmithRuntime/Source/Private/DatasmithRuntimeUtils.cpp
@@ -55,7 +55,7 @@ namespace DatasmithRuntime
 		for (; FirstEmptyUVs < NumberOfUVs; ++FirstEmptyUVs)
 		{
 			const TVertexInstanceAttributesConstRef<FVector2f> UVChannels = MeshDescription.VertexInstanceAttributes().GetAttributesRef<FVector2f>(MeshAttribute::VertexInstance::TextureCoordinate);
-			const FVector2D DefValue = UVChannels.GetDefaultValue();
+			const FVector2f DefValue = UVChannels.GetDefaultValue();
 			bool bHasNonDefaultValue = false;
 
 			for (FVertexInstanceID InstanceID : MeshDescription.VertexInstances().GetElementIDs())
diff --git a/Engine/Plugins/Experimental/GPULightmass/Source/GPULightmass/Private/PathTracingLightParameters.inl b/Engine/Plugins/Experimental/GPULightmass/Source/GPULightmass/Private/PathTracingLightParameters.inl
--- a/Engine/Plugins/Experimental/GPULightmass/Source/GPULightmass/Private/PathTracingLightParameters.inl
+++ b/Engine/Plugins/Experimental/GPULightmass/Source/GPULightmass/Private/PathTracingLightParameters.inl
@@ -113,9 +113,9 @@ void SetupPathTracingLightParameters(
 		DestLight.Normal = Light.Direction;
 		DestLight.dPdu = FVector::CrossProduct(Light.Tangent, Light.Direction);
 		DestLight.dPdv = Light.Tangent;
-		DestLight.Color = FVector(Light.Color);
-		DestLight.Dimensions = FVector(Light.SourceRadius, Light.SourceSoftRadius, Light.SourceLength);
-		DestLight.Shaping = Light.SpotAngles;
+		DestLight.Color = FVector3f(Light.Color);
+		DestLight.Dimensions = FVector3f(Light.SourceRadius, Light.SourceSoftRadius, Light.SourceLength);
+		DestLight.Shaping = FVector2f(Light.SpotAngles);
 		DestLight.Attenuation = 1.0f / Light.AttenuationRadius;
 		DestLight.FalloffExponent = Light.FalloffExponent;
 
@@ -167,11 +167,11 @@ void SetupPathTracingLightParameters(
 
 		FLinearColor LightColor = Light.Color;
 		LightColor /= 0.5f * Light.SourceWidth * Light.SourceHeight;
-		DestLight.Color = FVector(LightColor);
+		DestLight.Color = FVector3f(LightColor);
 
-		DestLight.Dimensions = FVector(Light.SourceWidth, Light.SourceHeight, 0.0f);
+		DestLight.Dimensions = FVector3f(Light.SourceWidth, Light.SourceHeight, 0.0f);
 		DestLight.Attenuation = 1.0f / Light.AttenuationRadius;
-		DestLight.Shaping = FVector2D(FMath::Cos(FMath::DegreesToRadians(Light.BarnDoorAngle)), Light.BarnDoorLength);
+		DestLight.Shaping = FVector2f(FMath::Cos(FMath::DegreesToRadians(Light.BarnDoorAngle)), Light.BarnDoorLength);
 
 		if (Light.IESTexture)
 		{
diff --git a/Engine/Plugins/Experimental/GPULightmass/Source/GPULightmass/Private/Scene/Lights.cpp b/Engine/Plugins/Experimental/GPULightmass/Source/GPULightmass/Private/Scene/Lights.cpp
--- a/Engine/Plugins/Experimental/GPULightmass/Source/GPULightmass/Private/Scene/Lights.cpp
+++ b/Engine/Plugins/Experimental/GPULightmass/Source/GPULightmass/Private/Scene/Lights.cpp
@@ -229,7 +229,7 @@ FLightRenderParameters FDirectionalLightRenderState::GetLightShaderParameters()
 	LightParameters.Direction = -Direction;
 	LightParameters.Tangent = -Direction;
 
-	LightParameters.SpotAngles = FVector2D(0, 0);
+	LightParameters.SpotAngles = FVector2f(0, 0);
 	LightParameters.SpecularScale = 0; // Irrelevant when tracing shadow rays
 	LightParameters.SourceRadius = FMath::Sin(0.5f * FMath::DegreesToRadians(LightSourceAngle));
 	LightParameters.SoftSourceRadius = 0; // Irrelevant when tracing shadow rays. FMath::Sin(0.5f * FMath::DegreesToRadians(LightSourceSoftAngle));
@@ -260,7 +260,7 @@ FLightRenderParameters FSpotLightRenderState::GetLightShaderParameters() const
 	LightParameters.WorldPosition = Position;
 	LightParameters.Direction = -Direction;
 	LightParameters.Tangent = Tangent;
-	LightParameters.SpotAngles = SpotAngles;
+	LightParameters.SpotAngles = FVector2f(SpotAngles);					// LWC_TODO: Precision loss?
 	LightParameters.InvRadius = 1.0f / AttenuationRadius;
 	LightParameters.Color = Color;
 	LightParameters.SourceRadius = SourceRadius;
diff --git a/Engine/Plugins/Experimental/GPULightmass/Source/GPULightmass/Private/Scene/Scene.cpp b/Engine/Plugins/Experimental/GPULightmass/Source/GPULightmass/Private/Scene/Scene.cpp
--- a/Engine/Plugins/Experimental/GPULightmass/Source/GPULightmass/Private/Scene/Scene.cpp
+++ b/Engine/Plugins/Experimental/GPULightmass/Source/GPULightmass/Private/Scene/Scene.cpp
@@ -762,7 +762,7 @@ void FScene::AddGeometryInstanceFromComponent(UStaticMeshComponent* InComponent)
 				Lightmap->Size,
 				FMath::Min((int32)FMath::CeilLogTwo((uint32)FMath::Min(Lightmap->GetPaddedSizeInTiles().X, Lightmap->GetPaddedSizeInTiles().Y)), GPreviewLightmapMipmapMaxLevel),
 				ResourceCluster, // temporarily promote unique ptr to raw ptr to make it copyable
-				FVector4f(Lightmap->LightmapObject->CoordinateScale, Lightmap->LightmapObject->CoordinateBias)
+				FVector4f(FVector2f(Lightmap->LightmapObject->CoordinateScale), FVector2f(Lightmap->LightmapObject->CoordinateBias))
 			};
 
 			InstanceLightmapRenderStateInitializers.Add(MoveTemp(Initializer));
@@ -987,7 +987,7 @@ void FScene::AddGeometryInstanceFromComponent(UInstancedStaticMeshComponent* InC
 					{
 						int32 X = RenderIndex % InstancesPerRow;
 						int32 Y = RenderIndex / InstancesPerRow;
-						FVector2D Bias = (FVector2D(X, Y) * FVector2D(BaseLightMapWidth, BaseLightMapHeight) + FVector2D(1, 1)) / Lightmap->Size;
+						FVector2f Bias = (FVector2f(X, Y) * FVector2f(BaseLightMapWidth, BaseLightMapHeight) + FVector2f(1, 1)) / Lightmap->Size;
 						Lightmap->MeshMapBuildData->PerInstanceLightmapData[GameThreadInstanceIndex].LightmapUVBias = Bias;
 						Lightmap->MeshMapBuildData->PerInstanceLightmapData[GameThreadInstanceIndex].ShadowmapUVBias = Bias;
 					}
@@ -1006,7 +1006,7 @@ void FScene::AddGeometryInstanceFromComponent(UInstancedStaticMeshComponent* InC
 				Lightmap->Size,
 				FMath::Min((int32)FMath::CeilLogTwo((uint32)FMath::Min(Lightmap->GetPaddedSizeInTiles().X, Lightmap->GetPaddedSizeInTiles().Y)), GPreviewLightmapMipmapMaxLevel),
 				ResourceCluster, // temporarily promote unique ptr to raw ptr to make it copyable
-				FVector4f(Lightmap->LightmapObject->CoordinateScale, Lightmap->LightmapObject->CoordinateBias)
+				FVector4f(FVector4(Lightmap->LightmapObject->CoordinateScale, Lightmap->LightmapObject->CoordinateBias))
 			};
 
 			InstanceLightmapRenderStateInitializers.Add(Initializer);
@@ -1219,7 +1219,7 @@ void FScene::AddGeometryInstanceFromComponent(ULandscapeComponent* InComponent)
 				Lightmap->Size,
 				FMath::Min((int32)FMath::CeilLogTwo((uint32)FMath::Min(Lightmap->GetPaddedSizeInTiles().X, Lightmap->GetPaddedSizeInTiles().Y)), GPreviewLightmapMipmapMaxLevel),
 				ResourceCluster, // temporarily promote unique ptr to raw ptr to make it copyable
-				FVector4f(Lightmap->LightmapObject->CoordinateScale, Lightmap->LightmapObject->CoordinateBias)
+				FVector4f(FVector4(Lightmap->LightmapObject->CoordinateScale, Lightmap->LightmapObject->CoordinateBias))
 			};
 
 			InstanceLightmapRenderStateInitializers.Add(Initializer);
diff --git a/Engine/Plugins/Experimental/GeometryScripting/Source/GeometryScriptingCore/Private/MeshUVFunctions.cpp b/Engine/Plugins/Experimental/GeometryScripting/Source/GeometryScriptingCore/Private/MeshUVFunctions.cpp
--- a/Engine/Plugins/Experimental/GeometryScripting/Source/GeometryScriptingCore/Private/MeshUVFunctions.cpp
+++ b/Engine/Plugins/Experimental/GeometryScripting/Source/GeometryScriptingCore/Private/MeshUVFunctions.cpp
@@ -192,7 +192,7 @@ UDynamicMesh* UGeometryScriptLibrary_MeshUVFunctions::ScaleMeshUVs(
 		return TargetMesh;
 	}
 
-	FVector2f UseScale = Scale;
+	FVector2f UseScale = FVector2f(Scale);
 	if (UseScale.Length() < 0.0001)
 	{
 		UseScale = FVector2f::One();
diff --git a/Engine/Plugins/Experimental/ImagePlate/Source/ImagePlate/Private/ImagePlateComponent.cpp b/Engine/Plugins/Experimental/ImagePlate/Source/ImagePlate/Private/ImagePlateComponent.cpp
--- a/Engine/Plugins/Experimental/ImagePlate/Source/ImagePlate/Private/ImagePlateComponent.cpp
+++ b/Engine/Plugins/Experimental/ImagePlate/Source/ImagePlate/Private/ImagePlateComponent.cpp
@@ -119,10 +119,10 @@ namespace
 			Vertices[2].Position  = FVector(0,  1.f,  1.f);
 			Vertices[3].Position  = FVector(0,  1.f, -1.f);
 
-			Vertices[0].TextureCoordinate[0] = FVector2D(0,0);
-			Vertices[1].TextureCoordinate[0] = FVector2D(0,1);
-			Vertices[2].TextureCoordinate[0] = FVector2D(1,0);
-			Vertices[3].TextureCoordinate[0] = FVector2D(1,1);
+			Vertices[0].TextureCoordinate[0] = FVector2f(0,0);
+			Vertices[1].TextureCoordinate[0] = FVector2f(0,1);
+			Vertices[2].TextureCoordinate[0] = FVector2f(1,0);
+			Vertices[3].TextureCoordinate[0] = FVector2f(1,1);
 
 			VertexBuffers.InitFromDynamicVertex(&VertexFactory, Vertices);
 
diff --git a/Engine/Plugins/Experimental/Interchange/Runtime/Source/Parsers/Fbx/Private/FbxMesh.cpp b/Engine/Plugins/Experimental/Interchange/Runtime/Source/Parsers/Fbx/Private/FbxMesh.cpp
--- a/Engine/Plugins/Experimental/Interchange/Runtime/Source/Parsers/Fbx/Private/FbxMesh.cpp
+++ b/Engine/Plugins/Experimental/Interchange/Runtime/Source/Parsers/Fbx/Private/FbxMesh.cpp
@@ -579,7 +579,7 @@ namespace UE
 							{
 								FUVID UVID = MeshDescription->CreateUV(UVLayerIndex);
 								FbxVector2 UVVector = FBXUVs.LayerElementUV[UVLayerIndex]->GetDirectArray().GetAt(UVIndex);
-								UVCoordinates[UVID] = FVector2D(static_cast<float>(UVVector[0]), 1.0f - static_cast<float>(UVVector[1]));	// flip the Y of UVs for DirectX
+								UVCoordinates[UVID] = FVector2f(static_cast<float>(UVVector[0]), 1.0f - static_cast<float>(UVVector[1]));	// flip the Y of UVs for DirectX
 							}
 						}
 					}
@@ -712,7 +712,7 @@ namespace UE
 										FinalUVVector.X = static_cast<float>(UVVector[0]);
 										FinalUVVector.Y = 1.f - static_cast<float>(UVVector[1]);   //flip the Y of UVs for DirectX
 									}
-									VertexInstanceUVs.Set(AddedVertexInstanceId, UVLayerIndex, FinalUVVector);
+									VertexInstanceUVs.Set(AddedVertexInstanceId, UVLayerIndex, FVector2f(FinalUVVector));	// LWC_TODO: Precision loss
 								}
 
 								//Color attribute
diff --git a/Engine/Plugins/Experimental/MeshModelingToolsetExp/Source/MeshModelingToolsEditorOnlyExp/Private/AttributeEditorTool.cpp b/Engine/Plugins/Experimental/MeshModelingToolsetExp/Source/MeshModelingToolsEditorOnlyExp/Private/AttributeEditorTool.cpp
--- a/Engine/Plugins/Experimental/MeshModelingToolsetExp/Source/MeshModelingToolsEditorOnlyExp/Private/AttributeEditorTool.cpp
+++ b/Engine/Plugins/Experimental/MeshModelingToolsetExp/Source/MeshModelingToolsEditorOnlyExp/Private/AttributeEditorTool.cpp
@@ -471,7 +471,7 @@ void UAttributeEditorTool::ClearUVs()
 			{
 				for (const FVertexInstanceID ElID : Instances.GetElementIDs())
 				{
-					InstanceUVs.Set(ElID, LayerIndex, FVector2D::ZeroVector);
+					InstanceUVs.Set(ElID, LayerIndex, FVector2f::ZeroVector);
 				}
 			}
 		}
@@ -520,7 +520,7 @@ void UAttributeEditorTool::DeleteSelectedUVSet()
 			const FVertexInstanceArray& Instances = EditedMesh.VertexInstances();
 			for (const FVertexInstanceID InstanceID : Instances.GetElementIDs())
 			{
-				InstanceUVs.Set(InstanceID, DeleteIndex, FVector2D::ZeroVector);
+				InstanceUVs.Set(InstanceID, DeleteIndex, FVector2f::ZeroVector);
 			}
 		}
 
@@ -597,7 +597,7 @@ void UAttributeEditorTool::DuplicateSelectedUVSet()
 			const FVertexInstanceArray& Instances = EditedMesh.VertexInstances();
 			for (const FVertexInstanceID InstanceID : Instances.GetElementIDs())
 			{
-				FVector2D SourceUV = InstanceUVs.Get(InstanceID, SourceIndex);
+				FVector2f SourceUV = InstanceUVs.Get(InstanceID, SourceIndex);
 				InstanceUVs.Set(InstanceID, NewChannelIndex, SourceUV);
 			}
 
diff --git a/Engine/Plugins/Experimental/MeshModelingToolsetExp/Source/MeshModelingToolsExp/Private/DisplaceMeshTool.cpp b/Engine/Plugins/Experimental/MeshModelingToolsetExp/Source/MeshModelingToolsExp/Private/DisplaceMeshTool.cpp
--- a/Engine/Plugins/Experimental/MeshModelingToolsetExp/Source/MeshModelingToolsExp/Private/DisplaceMeshTool.cpp
+++ b/Engine/Plugins/Experimental/MeshModelingToolsetExp/Source/MeshModelingToolsExp/Private/DisplaceMeshTool.cpp
@@ -1077,7 +1077,7 @@ void UDisplaceMeshTool::OnPropertyModified(UObject* PropertySet, FProperty* Prop
 		{
 			DisplacerDownCast->SetFilterDirection(DirectionalFilterProperties->FilterDirection);
 			DisplacerDownCast->SetSineWaveDirection(SineWaveProperties->SineWaveDirection);
-			DisplacerDownCast->SetDisplacementMapUVAdjustment(TextureMapProperties->UVScale, TextureMapProperties->UVOffset);
+			DisplacerDownCast->SetDisplacementMapUVAdjustment(FVector2f(TextureMapProperties->UVScale), FVector2f(TextureMapProperties->UVOffset));	// LWC_TODO: Precision loss
 		}
 
 		StartComputation();
diff --git a/Engine/Plugins/Experimental/MotoSynth/Source/MotoSynth/Private/MotoSynthEngine.cpp b/Engine/Plugins/Experimental/MotoSynth/Source/MotoSynth/Private/MotoSynthEngine.cpp
--- a/Engine/Plugins/Experimental/MotoSynth/Source/MotoSynth/Private/MotoSynthEngine.cpp
+++ b/Engine/Plugins/Experimental/MotoSynth/Source/MotoSynth/Private/MotoSynthEngine.cpp
@@ -123,7 +123,7 @@ void FMotoSynthEngine::SetSourceData(uint32 InAccelerationSourceID, uint32 InDec
 		FVector2f DecelRPMRange;
 		InDecelerationSourceData->RPMCurve.GetValueRange(DecelRPMRange.X, DecelRPMRange.Y);
 
-		FVector2D NewRPMRange = { FMath::Max(AccelRPMRange.X, DecelRPMRange.X), FMath::Min(AccelRPMRange.Y, DecelRPMRange.Y) };
+		FVector2f NewRPMRange = { FMath::Max(AccelRPMRange.X, DecelRPMRange.X), FMath::Min(AccelRPMRange.Y, DecelRPMRange.Y) };
 		RPMRange = NewRPMRange;
 
 		SynthCommand([this, InAccelerationSourceData, InDecelerationSourceData, NewRPMRange]() mutable
@@ -139,7 +139,7 @@ void FMotoSynthEngine::SetSourceData(uint32 InAccelerationSourceID, uint32 InDec
 
 void FMotoSynthEngine::GetRPMRange(FVector2D& OutRPMRange)
 {
-	OutRPMRange = RPMRange;
+	OutRPMRange = FVector2D(RPMRange);
 }
 
 void FMotoSynthEngine::SetSettings(const FMotoSynthRuntimeSettings& InSettings)
diff --git a/Engine/Plugins/Experimental/PlanarCutPlugin/Source/PlanarCut/Private/FractureAutoUV.cpp b/Engine/Plugins/Experimental/PlanarCutPlugin/Source/PlanarCut/Private/FractureAutoUV.cpp
--- a/Engine/Plugins/Experimental/PlanarCutPlugin/Source/PlanarCut/Private/FractureAutoUV.cpp
+++ b/Engine/Plugins/Experimental/PlanarCutPlugin/Source/PlanarCut/Private/FractureAutoUV.cpp
@@ -324,7 +324,7 @@ struct FGeomFlatUVMesh
 
 	inline FVector3d GetVertex(int32 VID) const
 	{
-		const FVector2D& UV = Collection->UVs[VID][UVLayer];
+		const FVector2f& UV = Collection->UVs[VID][UVLayer];
 		return FVector3d(UV.X, UV.Y, 0);
 	}
 
diff --git a/Engine/Plugins/Experimental/PlanarCutPlugin/Source/PlanarCut/Private/GeometryMeshConversion.cpp b/Engine/Plugins/Experimental/PlanarCutPlugin/Source/PlanarCut/Private/GeometryMeshConversion.cpp
--- a/Engine/Plugins/Experimental/PlanarCutPlugin/Source/PlanarCut/Private/GeometryMeshConversion.cpp
+++ b/Engine/Plugins/Experimental/PlanarCutPlugin/Source/PlanarCut/Private/GeometryMeshConversion.cpp
@@ -2967,7 +2967,7 @@ bool FDynamicMeshCollection::UpdateCollection(const FTransform& FromCollection,
 		{
 			FVector2f UV;
 			AugmentedDynamicMesh::GetUV(Mesh, VID, UV, UVLayer);
-			Output.UVs[CopyToIdx][UVLayer] = FVector2D(UV);
+			Output.UVs[CopyToIdx][UVLayer] = UV;
 		}
 		
 		FVector3f TangentU, TangentV;
@@ -3069,7 +3069,7 @@ int32 FDynamicMeshCollection::AppendToCollection(const FTransform& FromCollectio
 		{
 			FVector2f UV;
 			AugmentedDynamicMesh::GetUV(Mesh, VID, UV, UVLayer);
-			Output.UVs[CopyToIdx][UVLayer] = FVector2D(UV);
+			Output.UVs[CopyToIdx][UVLayer] = UV;
 		}
 
 		FVector3f TangentU, TangentV;
diff --git a/Engine/Plugins/Experimental/PlanarCutPlugin/Source/PlanarCut/Private/PlanarCut.cpp b/Engine/Plugins/Experimental/PlanarCutPlugin/Source/PlanarCut/Private/PlanarCut.cpp
--- a/Engine/Plugins/Experimental/PlanarCutPlugin/Source/PlanarCut/Private/PlanarCut.cpp
+++ b/Engine/Plugins/Experimental/PlanarCutPlugin/Source/PlanarCut/Private/PlanarCut.cpp
@@ -110,11 +110,11 @@ void FInternalSurfaceMaterials::SetUVScaleFromCollection(const FGeometryCollecti
 	{
 		const FIntVector& Tri = Collection.Indices[FaceIdx];
 		WorldDistance += FVector::Distance(Collection.Vertex[Tri.X], Collection.Vertex[Tri.Y]);
-		UVDistance += FVector2D::Distance(Collection.UVs[Tri.X][0], Collection.UVs[Tri.Y][0]);
+		UVDistance += FVector2D::Distance(FVector2D(Collection.UVs[Tri.X][0]), FVector2D(Collection.UVs[Tri.Y][0]));
 		WorldDistance += FVector::Distance(Collection.Vertex[Tri.Z], Collection.Vertex[Tri.Y]);
-		UVDistance += FVector2D::Distance(Collection.UVs[Tri.Z][0], Collection.UVs[Tri.Y][0]);
+		UVDistance += FVector2D::Distance(FVector2D(Collection.UVs[Tri.Z][0]), FVector2D(Collection.UVs[Tri.Y][0]));
 		WorldDistance += FVector::Distance(Collection.Vertex[Tri.X], Collection.Vertex[Tri.Z]);
-		UVDistance += FVector2D::Distance(Collection.UVs[Tri.X][0], Collection.UVs[Tri.Z][0]);
+		UVDistance += FVector2D::Distance(FVector2D(Collection.UVs[Tri.X][0]), FVector2D(Collection.UVs[Tri.Z][0]));
 	}
 
 	if (WorldDistance > 0)
diff --git a/Engine/Plugins/Experimental/ProxyLODPlugin/Source/ProxyLOD/Private/ProxyLODMeshAttrTransfer.cpp b/Engine/Plugins/Experimental/ProxyLODPlugin/Source/ProxyLOD/Private/ProxyLODMeshAttrTransfer.cpp
--- a/Engine/Plugins/Experimental/ProxyLODPlugin/Source/ProxyLOD/Private/ProxyLODMeshAttrTransfer.cpp
+++ b/Engine/Plugins/Experimental/ProxyLODPlugin/Source/ProxyLOD/Private/ProxyLODMeshAttrTransfer.cpp
@@ -112,7 +112,7 @@ void ProxyLOD::TransferMeshAttributes(const FClosestPolyField& SrcPolyField, FMe
 					VertexInstanceBinormalSigns[Idx] = GetBasisDeterminantSign(VertexInstanceTangents[Idx], AverageUnitVector(RawPoly.WedgeTangentY), VertexInstanceNormals[Idx]);
 
 					// Average Texture Coords
-					VertexInstanceUVs.Set(Idx, 0, AverageTexCoord(RawPoly.WedgeTexCoords[0]));
+					VertexInstanceUVs.Set(Idx, 0, FVector2f(AverageTexCoord(RawPoly.WedgeTexCoords[0])));	// LWC_TODO: Precision loss
 				}
 			}
 			// Assign the material index that the last vertex of this face sees.
diff --git a/Engine/Plugins/Experimental/ProxyLODPlugin/Source/ProxyLOD/Private/ProxyLODMeshTypes.cpp b/Engine/Plugins/Experimental/ProxyLODPlugin/Source/ProxyLOD/Private/ProxyLODMeshTypes.cpp
--- a/Engine/Plugins/Experimental/ProxyLODPlugin/Source/ProxyLOD/Private/ProxyLODMeshTypes.cpp
+++ b/Engine/Plugins/Experimental/ProxyLODPlugin/Source/ProxyLOD/Private/ProxyLODMeshTypes.cpp
@@ -372,7 +372,7 @@ FMeshDescriptionArrayAdapter::FRawPoly FMeshDescriptionArrayAdapter::GetRawPoly(
 				{
 					if (AttributesGetter->VertexInstanceUVs.GetNumChannels() > Idx)
 					{
-						RawPoly.WedgeTexCoords[Idx][i] = AttributesGetter->VertexInstanceUVs.Get(VertexInstanceIDs[i], Idx);
+						RawPoly.WedgeTexCoords[Idx][i] = FVector2D(AttributesGetter->VertexInstanceUVs.Get(VertexInstanceIDs[i], Idx));
 					}
 					else
 					{
diff --git a/Engine/Plugins/Experimental/RemoteSession/Source/RemoteSession/Private/Channels/RemoteSessionARCameraChannel.cpp b/Engine/Plugins/Experimental/RemoteSession/Source/RemoteSession/Private/Channels/RemoteSessionARCameraChannel.cpp
--- a/Engine/Plugins/Experimental/RemoteSession/Source/RemoteSession/Private/Channels/RemoteSessionARCameraChannel.cpp
+++ b/Engine/Plugins/Experimental/RemoteSession/Source/RemoteSession/Private/Channels/RemoteSessionARCameraChannel.cpp
@@ -195,16 +195,16 @@ void FARCameraSceneViewExtension::PreRenderView_RenderThread(FRHICommandListImme
 		Vertices.SetNumUninitialized(4);
 
 		Vertices[0].Position = FVector4f(0.f, 0.f, 0.f, 1.f);
-		Vertices[0].UV = FVector2D(0.f, 0.f);
+		Vertices[0].UV = FVector2f(0.f, 0.f);
 
 		Vertices[1].Position = FVector4f(1.f, 0.f, 0.f, 1.f);
-		Vertices[1].UV = FVector2D(1.f, 0.f);
+		Vertices[1].UV = FVector2f(1.f, 0.f);
 
 		Vertices[2].Position = FVector4f(0.f, 1.f, 0.f, 1.f);
-		Vertices[2].UV = FVector2D(0.f, 1.f);
+		Vertices[2].UV = FVector2f(0.f, 1.f);
 
 		Vertices[3].Position = FVector4f(1.f, 1.f, 0.f, 1.f);
-		Vertices[3].UV = FVector2D(1.f, 1.f);
+		Vertices[3].UV = FVector2f(1.f, 1.f);
 
 		FRHIResourceCreateInfo CreateInfoVB(TEXT("FARCameraSceneViewExtension"), &Vertices);
 		VertexBufferRHI = RHICreateVertexBuffer(Vertices.GetResourceDataSize(), BUF_Static, CreateInfoVB);
diff --git a/Engine/Plugins/Experimental/SimpleHMD/Source/SimpleHMD/Private/SimpleHMD.cpp b/Engine/Plugins/Experimental/SimpleHMD/Source/SimpleHMD/Private/SimpleHMD.cpp
--- a/Engine/Plugins/Experimental/SimpleHMD/Source/SimpleHMD/Private/SimpleHMD.cpp
+++ b/Engine/Plugins/Experimental/SimpleHMD/Source/SimpleHMD/Private/SimpleHMD.cpp
@@ -189,17 +189,17 @@ void FSimpleHMD::DrawDistortionMesh_RenderThread(struct FHeadMountedDisplayPassC
 	static const FDistortionVertex MeshVerts[][NumVerts] = {
 		{
 			// left eye
-			{ FVector2D(-0.9f, -0.9f), FVector2D(0.0f, 1.0f), FVector2D(0.0f, 1.0f), FVector2D(0.0f, 1.0f), 1.0f, 0.0f },
-			{ FVector2D(-0.1f, -0.9f), FVector2D(0.5f, 1.0f), FVector2D(0.5f, 1.0f), FVector2D(0.5f, 1.0f), 1.0f, 0.0f },
-			{ FVector2D(-0.1f, 0.9f), FVector2D(0.5f, 0.0f), FVector2D(0.5f, 0.0f), FVector2D(0.5f, 0.0f), 1.0f, 0.0f },
-			{ FVector2D(-0.9f, 0.9f), FVector2D(0.0f, 0.0f), FVector2D(0.0f, 0.0f), FVector2D(0.0f, 0.0f), 1.0f, 0.0f },
+			{ FVector2f(-0.9f, -0.9f), FVector2f(0.0f, 1.0f), FVector2f(0.0f, 1.0f), FVector2f(0.0f, 1.0f), 1.0f, 0.0f },
+			{ FVector2f(-0.1f, -0.9f), FVector2f(0.5f, 1.0f), FVector2f(0.5f, 1.0f), FVector2f(0.5f, 1.0f), 1.0f, 0.0f },
+			{ FVector2f(-0.1f, 0.9f), FVector2f(0.5f, 0.0f), FVector2f(0.5f, 0.0f), FVector2f(0.5f, 0.0f), 1.0f, 0.0f },
+			{ FVector2f(-0.9f, 0.9f), FVector2f(0.0f, 0.0f), FVector2f(0.0f, 0.0f), FVector2f(0.0f, 0.0f), 1.0f, 0.0f },
 		},
 		{
 			// right eye
-			{ FVector2D(0.1f, -0.9f), FVector2D(0.5f, 1.0f), FVector2D(0.5f, 1.0f), FVector2D(0.5f, 1.0f), 1.0f, 0.0f },
-			{ FVector2D(0.9f, -0.9f), FVector2D(1.0f, 1.0f), FVector2D(1.0f, 1.0f), FVector2D(1.0f, 1.0f), 1.0f, 0.0f },
-			{ FVector2D(0.9f, 0.9f), FVector2D(1.0f, 0.0f), FVector2D(1.0f, 0.0f), FVector2D(1.0f, 0.0f), 1.0f, 0.0f },
-			{ FVector2D(0.1f, 0.9f), FVector2D(0.5f, 0.0f), FVector2D(0.5f, 0.0f), FVector2D(0.5f, 0.0f), 1.0f, 0.0f },
+			{ FVector2f(0.1f, -0.9f), FVector2f(0.5f, 1.0f), FVector2f(0.5f, 1.0f), FVector2f(0.5f, 1.0f), 1.0f, 0.0f },
+			{ FVector2f(0.9f, -0.9f), FVector2f(1.0f, 1.0f), FVector2f(1.0f, 1.0f), FVector2f(1.0f, 1.0f), 1.0f, 0.0f },
+			{ FVector2f(0.9f, 0.9f), FVector2f(1.0f, 0.0f), FVector2f(1.0f, 0.0f), FVector2f(1.0f, 0.0f), 1.0f, 0.0f },
+			{ FVector2f(0.1f, 0.9f), FVector2f(0.5f, 0.0f), FVector2f(0.5f, 0.0f), FVector2f(0.5f, 0.0f), 1.0f, 0.0f },
 		}
 	};
 
diff --git a/Engine/Plugins/Experimental/Text3D/Source/Text3D/Private/Data.cpp b/Engine/Plugins/Experimental/Text3D/Source/Text3D/Private/Data.cpp
--- a/Engine/Plugins/Experimental/Text3D/Source/Text3D/Private/Data.cpp
+++ b/Engine/Plugins/Experimental/Text3D/Source/Text3D/Private/Data.cpp
@@ -101,7 +101,7 @@ int32 FData::AddVertex(const FVector& Position, const FVector& TangentX, const F
 	StaticMeshAttributes.GetVertexPositions()[FVertexID(VertexIndex)] = Position;
 	const FVertexInstanceID Instance(static_cast<uint32>(VertexIndex));
 
-	StaticMeshAttributes.GetVertexInstanceUVs()[Instance] = TextureCoordinates;
+	StaticMeshAttributes.GetVertexInstanceUVs()[Instance] = FVector2f(TextureCoordinates);
 	StaticMeshAttributes.GetVertexInstanceNormals()[Instance] = TangentZ;
 	StaticMeshAttributes.GetVertexInstanceTangents()[Instance] = TangentX;
 
diff --git a/Engine/Plugins/Experimental/Text3D/Source/Text3D/Private/GlyphLoader.cpp b/Engine/Plugins/Experimental/Text3D/Source/Text3D/Private/GlyphLoader.cpp
--- a/Engine/Plugins/Experimental/Text3D/Source/Text3D/Private/GlyphLoader.cpp
+++ b/Engine/Plugins/Experimental/Text3D/Source/Text3D/Private/GlyphLoader.cpp
@@ -52,7 +52,7 @@ void FGlyphLoader::FLine::Add(FGlyphLoader* const Loader)
 {
 	if (!Loader->ProcessedContour.Num() || !(Position - Loader->FirstPosition).IsNearlyZero())
 	{
-		Loader->ProcessedContour.AddHead(Position);
+		Loader->ProcessedContour.AddHead(FVector2f(Position));
 	}
 }
 
@@ -105,7 +105,7 @@ void FGlyphLoader::FCurve::Add(FGlyphLoader* const LoaderIn)
 	Start.T = StartT;
 	Start.Position = Position(Start.T);
 	Start.Tangent = Tangent(Start.T);
-	Loader->ProcessedContour.AddHead(Start.Position);
+	Loader->ProcessedContour.AddHead(FVector2f(Start.Position));
 	Start.Point = Loader->ProcessedContour.GetHead();
 	First = Start.Point;
 	bFirstSplit = false;
@@ -158,7 +158,7 @@ void FGlyphLoader::FCurve::Split(const FPointData& Start, const FPointData& End)
 	Middle.T = (Start.T + End.T) / 2.f;
 	Middle.Position = Position(Middle.T);
 	Middle.Tangent = Tangent(Middle.T);
-	Loader->ProcessedContour.InsertNode(Middle.Position, Start.Point);
+	Loader->ProcessedContour.InsertNode(FVector2f(Middle.Position), Start.Point);
 	Middle.Point = Start.Point->GetPrevNode();
 
 	CheckPart(Start, Middle);
diff --git a/Engine/Plugins/Experimental/Text3D/Source/Text3D/Private/MeshCreator.cpp b/Engine/Plugins/Experimental/Text3D/Source/Text3D/Private/MeshCreator.cpp
--- a/Engine/Plugins/Experimental/Text3D/Source/Text3D/Private/MeshCreator.cpp
+++ b/Engine/Plugins/Experimental/Text3D/Source/Text3D/Private/MeshCreator.cpp
@@ -325,7 +325,7 @@ void FMeshCreator::MirrorGroup(const EText3DGroupType TypeIn, const EText3DGroup
 		const FVector Normal = VertexNormals[InstanceID];
 		const FVector Tangent = VertexTangents[InstanceID];
 
-		Data->AddVertex({ Extrude - Position.X, Position.Y, Position.Z }, { -Tangent.X, Tangent.Y, Tangent.Z }, { -Normal.X, Normal.Y, Normal.Z }, VertexUVs[InstanceID]);
+		Data->AddVertex({ Extrude - Position.X, Position.Y, Position.Z }, { -Tangent.X, Tangent.Y, Tangent.Z }, { -Normal.X, Normal.Y, Normal.Z }, FVector2D(VertexUVs[InstanceID]));
 	}
 
 	Data->AddTriangles(TrianglesInNum);
diff --git a/Engine/Plugins/Experimental/UVEditor/Source/UVEditor/Private/UVEditorMode.cpp b/Engine/Plugins/Experimental/UVEditor/Source/UVEditor/Private/UVEditorMode.cpp
--- a/Engine/Plugins/Experimental/UVEditor/Source/UVEditor/Private/UVEditorMode.cpp
+++ b/Engine/Plugins/Experimental/UVEditor/Source/UVEditor/Private/UVEditorMode.cpp
@@ -474,7 +474,7 @@ void UUVEditorMode::InitializeTargets(const TArray<TObjectPtr<UObject>>& AssetsI
 	};
 	auto VertPositionToUV = [this, ScaleFactor](const FVector3d& VertPosition)
 	{
-		return FVector2D(VertPosition.Y / ScaleFactor, 1 - (VertPosition.X / ScaleFactor));
+		return FVector2f(VertPosition.Y / ScaleFactor, 1 - (VertPosition.X / ScaleFactor));
 	};
 
 	// Construct the full input objects that the tools actually operate on.
diff --git a/Engine/Plugins/Experimental/VirtualHeightfieldMesh/Source/VirtualHeightfieldMesh/Private/VirtualHeightfieldMeshSceneProxy.cpp b/Engine/Plugins/Experimental/VirtualHeightfieldMesh/Source/VirtualHeightfieldMesh/Private/VirtualHeightfieldMeshSceneProxy.cpp
--- a/Engine/Plugins/Experimental/VirtualHeightfieldMesh/Source/VirtualHeightfieldMesh/Private/VirtualHeightfieldMeshSceneProxy.cpp
+++ b/Engine/Plugins/Experimental/VirtualHeightfieldMesh/Source/VirtualHeightfieldMesh/Private/VirtualHeightfieldMeshSceneProxy.cpp
@@ -463,7 +463,7 @@ void FVirtualHeightfieldMeshSceneProxy::CreateRenderThreadResources()
 				UniformParams.PageTableSize = FVector4f(PageTableSizeX, PageTableSizeY, 1.f / PageTableSizeX, 1.f / PageTableSizeY);
 
 				const float PhysicalTextureSize = AllocatedVirtualTexture->GetPhysicalTextureSize(0);
-				UniformParams.PhysicalTextureSize = FVector2D(PhysicalTextureSize, 1.f / PhysicalTextureSize);
+				UniformParams.PhysicalTextureSize = FVector2f(PhysicalTextureSize, 1.f / PhysicalTextureSize);
 
 				UniformParams.VirtualHeightfieldToLocal = FMatrix44f(UVToLocal);
 				UniformParams.VirtualHeightfieldToWorld = FMatrix44f(UVToWorld);		// LWC_TODO: Precision loss
diff --git a/Engine/Plugins/Experimental/VirtualProduction/LedWallCalibration/Source/LedWallCalibration/Private/LedWallCalibration.cpp b/Engine/Plugins/Experimental/VirtualProduction/LedWallCalibration/Source/LedWallCalibration/Private/LedWallCalibration.cpp
--- a/Engine/Plugins/Experimental/VirtualProduction/LedWallCalibration/Source/LedWallCalibration/Private/LedWallCalibration.cpp
+++ b/Engine/Plugins/Experimental/VirtualProduction/LedWallCalibration/Source/LedWallCalibration/Private/LedWallCalibration.cpp
@@ -210,15 +210,15 @@ namespace LedWallCalibration
 			// Now we just need to find the 4th vertex/uv. 
 			// In this algorithm, we assume they are rectangular. So it is a matter of adding the 2 vectors.
 
-			const float DistanceUV01 = FVector2D::Distance(WallPanel.UVs[0], WallPanel.UVs[1]);
-			const float DistanceUV02 = FVector2D::Distance(WallPanel.UVs[0], WallPanel.UVs[2]);
-			const float DistanceUV12 = FVector2D::Distance(WallPanel.UVs[1], WallPanel.UVs[2]);
+			const float DistanceUV01 = FVector2f::Distance(WallPanel.UVs[0], WallPanel.UVs[1]);
+			const float DistanceUV02 = FVector2f::Distance(WallPanel.UVs[0], WallPanel.UVs[2]);
+			const float DistanceUV12 = FVector2f::Distance(WallPanel.UVs[1], WallPanel.UVs[2]);
 
 			const float DistanceMax = FMath::Max(DistanceUV01, FMath::Max(DistanceUV02, DistanceUV12));
 
 			if (FMath::IsNearlyEqual(DistanceMax, DistanceUV01))
 			{
-				const FVector2D& OriginUV = WallPanel.UVs[2];
+				const FVector2f& OriginUV = WallPanel.UVs[2];
 				const FVector& Origin3d = WallPanel.Vertices[2];
 
 				WallPanel.UVs[3] = OriginUV + (WallPanel.UVs[0] - OriginUV) + (WallPanel.UVs[1] - OriginUV);
@@ -226,7 +226,7 @@ namespace LedWallCalibration
 			}
 			else if (FMath::IsNearlyEqual(DistanceMax, DistanceUV02))
 			{
-				const FVector2D& OriginUV = WallPanel.UVs[1];
+				const FVector2f& OriginUV = WallPanel.UVs[1];
 				const FVector& Origin3d = WallPanel.Vertices[1];
 
 				WallPanel.UVs[3] = OriginUV + (WallPanel.UVs[0] - OriginUV) + (WallPanel.UVs[2] - OriginUV);
@@ -234,7 +234,7 @@ namespace LedWallCalibration
 			}
 			else
 			{
-				const FVector2D& OriginUV = WallPanel.UVs[0];
+				const FVector2f& OriginUV = WallPanel.UVs[0];
 				const FVector& Origin3d = WallPanel.Vertices[0];
 
 				WallPanel.UVs[3] = OriginUV + (WallPanel.UVs[1] - OriginUV) + (WallPanel.UVs[2] - OriginUV);
diff --git a/Engine/Plugins/Experimental/Water/Source/Runtime/Public/WaterVertexFactory.inl b/Engine/Plugins/Experimental/Water/Source/Runtime/Public/WaterVertexFactory.inl
--- a/Engine/Plugins/Experimental/Water/Source/Runtime/Public/WaterVertexFactory.inl
+++ b/Engine/Plugins/Experimental/Water/Source/Runtime/Public/WaterVertexFactory.inl
@@ -119,7 +119,7 @@ void TWaterVertexFactory<bWithWaterSelectionSupport>::SetupUniformDataForGroup(E
 	FWaterVertexFactoryParameters UniformParams;
 	UniformParams.NumQuadsPerTileSide = NumQuadsPerSide;
 	UniformParams.LODScale = LODScale;
-	UniformParams.MorphOrigin = MorphOrigin;
+	UniformParams.MorphOrigin = FVector2f(MorphOrigin);	// LWC_TODO: Precision loss
 	UniformParams.bRenderSelected = true;
 	UniformParams.bRenderUnselected = true;
 
diff --git a/Engine/Plugins/FX/Niagara/Source/Niagara/Classes/NDISkeletalMeshCommon.h b/Engine/Plugins/FX/Niagara/Source/Niagara/Classes/NDISkeletalMeshCommon.h
--- a/Engine/Plugins/FX/Niagara/Source/Niagara/Classes/NDISkeletalMeshCommon.h
+++ b/Engine/Plugins/FX/Niagara/Source/Niagara/Classes/NDISkeletalMeshCommon.h
@@ -442,11 +442,11 @@ struct FSkelMeshVertexAccessor
 		check(LODData);
 		if (bUseFullPrecisionUVs)
 		{
-			return LODData->StaticVertexBuffers.StaticMeshVertexBuffer.GetVertexUV_Typed<EStaticMeshVertexUVType::HighPrecision>(VertexIdx, UVChannel);
+			return FVector2D(LODData->StaticVertexBuffers.StaticMeshVertexBuffer.GetVertexUV_Typed<EStaticMeshVertexUVType::HighPrecision>(VertexIdx, UVChannel));
 		}
 		else
 		{
-			return LODData->StaticVertexBuffers.StaticMeshVertexBuffer.GetVertexUV_Typed<EStaticMeshVertexUVType::Default>(VertexIdx, UVChannel);
+			return FVector2D(LODData->StaticVertexBuffers.StaticMeshVertexBuffer.GetVertexUV_Typed<EStaticMeshVertexUVType::Default>(VertexIdx, UVChannel));
 		}
 	}
 
diff --git a/Engine/Plugins/FX/Niagara/Source/Niagara/Classes/NiagaraDataInterface.h b/Engine/Plugins/FX/Niagara/Source/Niagara/Classes/NiagaraDataInterface.h
--- a/Engine/Plugins/FX/Niagara/Source/Niagara/Classes/NiagaraDataInterface.h
+++ b/Engine/Plugins/FX/Niagara/Source/Niagara/Classes/NiagaraDataInterface.h
@@ -733,7 +733,7 @@ struct TNDIRandomHelper : public TRandomPolicy
 	{
 		//TODO: This is gonna be slooooow. Move to an LUT possibly or find faster method.
 		//Can probably handle lower quality randoms / uniformity for a decent speed win.
-		FVector2f r = TRandomPolicy::Rand2(InstanceIndex);
+		FVector2f r = FVector2f(TRandomPolicy::Rand2(InstanceIndex));
 		float sqrt0 = FMath::Sqrt(r.X);
 		return FVector3f(1.0f - sqrt0, sqrt0 * (1.0 - r.Y), r.Y * sqrt0);
 	}
diff --git a/Engine/Plugins/FX/Niagara/Source/Niagara/Private/DataInterface/NiagaraDataInterfaceStaticMesh.cpp b/Engine/Plugins/FX/Niagara/Source/Niagara/Private/DataInterface/NiagaraDataInterfaceStaticMesh.cpp
--- a/Engine/Plugins/FX/Niagara/Source/Niagara/Private/DataInterface/NiagaraDataInterfaceStaticMesh.cpp
+++ b/Engine/Plugins/FX/Niagara/Source/Niagara/Private/DataInterface/NiagaraDataInterfaceStaticMesh.cpp
@@ -1037,11 +1037,11 @@ namespace NDIStaticMeshLocal
 
 		FORCEINLINE FVector2D GetTriangleUV(const FVector3f& BaryCoord, int32 Index0, int32 Index1, int32 Index2, int32 UVSet) const
 		{
-			FVector2D UV;
+			FVector2f UV;
 			UV = LODResource->VertexBuffers.StaticMeshVertexBuffer.GetVertexUV(Index0, UVSet) * BaryCoord.X;
 			UV += LODResource->VertexBuffers.StaticMeshVertexBuffer.GetVertexUV(Index1, UVSet) * BaryCoord.Y;
 			UV += LODResource->VertexBuffers.StaticMeshVertexBuffer.GetVertexUV(Index2, UVSet) * BaryCoord.Z;
-			return UV;
+			return FVector2D(UV);
 		}
 
 		FORCEINLINE FVector3f GetPosition(int32 Vertex) const
@@ -1084,7 +1084,7 @@ namespace NDIStaticMeshLocal
 
 		FORCEINLINE FVector2D GetUV(int32 Vertex, int32 UVSet) const
 		{
-			return LODResource->VertexBuffers.StaticMeshVertexBuffer.GetVertexUV(Vertex, UVSet);
+			return FVector2D(LODResource->VertexBuffers.StaticMeshVertexBuffer.GetVertexUV(Vertex, UVSet));
 		}
 
 		FORCEINLINE bool IsCpuUniformlyDistributedSampling() const
@@ -2898,7 +2898,7 @@ void UNiagaraDataInterfaceStaticMesh::VMGetVertexUV(FVectorVMExternalFunctionCon
 		{
 			const int32 Vertex = VertexParam.GetAndAdvance();
 			const int32 UVSet = UVSetParam.GetAndAdvance();
-			const FVector2f UV = StaticMeshHelper.GetUV(FMath::Clamp(Vertex, 0, VertexMax), FMath::Clamp(UVSet, 0, UVSetMax));
+			const FVector2f UV = FVector2f(StaticMeshHelper.GetUV(FMath::Clamp(Vertex, 0, VertexMax), FMath::Clamp(UVSet, 0, UVSetMax)));	// LWC_TODO: Precision loss
 			OutUV.SetAndAdvance(UV);
 		}
 	}
@@ -3254,7 +3254,7 @@ void UNiagaraDataInterfaceStaticMesh::VMGetTriangleUV(FVectorVMExternalFunctionC
 			const int32 Triangle = FMath::Clamp(TriangleParam.GetAndAdvance(), 0, TriangleMax);
 			const FVector3f BaryCoord = BaryCoordParam.GetAndAdvance();
 			const int32 UVSet = FMath::Clamp(UVSetParam.GetAndAdvance(), 0, UVSetMax);
-			const FVector2f UV = StaticMeshHelper.GetTriangleUV(BaryCoord, IndexArray[Triangle * 3 + 0], IndexArray[Triangle * 3 + 1], IndexArray[Triangle * 3 + 2], UVSet);
+			const FVector2f UV = FVector2f(StaticMeshHelper.GetTriangleUV(BaryCoord, IndexArray[Triangle * 3 + 0], IndexArray[Triangle * 3 + 1], IndexArray[Triangle * 3 + 2], UVSet));	// LWC_TODO: Precision loss
 			OutUVParam.SetAndAdvance(UV);
 		}
 	}
diff --git a/Engine/Plugins/FX/Niagara/Source/Niagara/Private/NDISkeletalMesh_TriangleSampling.cpp b/Engine/Plugins/FX/Niagara/Source/Niagara/Private/NDISkeletalMesh_TriangleSampling.cpp
--- a/Engine/Plugins/FX/Niagara/Source/Niagara/Private/NDISkeletalMesh_TriangleSampling.cpp
+++ b/Engine/Plugins/FX/Niagara/Source/Niagara/Private/NDISkeletalMesh_TriangleSampling.cpp
@@ -1025,9 +1025,9 @@ void UNiagaraDataInterfaceSkeletalMesh::GetTriCoordUV(FVectorVMExternalFunctionC
 				const int32 Idx1 = IndexBuffer->Get(Tri + 1);
 				const int32 Idx2 = IndexBuffer->Get(Tri + 2);
 				const int32 UVSet = FMath::Clamp(UVSetParam.GetAndAdvance(), 0, UVSetMax);
-				const FVector2f UV0 = VertAccessor.GetVertexUV(LODData, Idx0, UVSet);
-				const FVector2f UV1 = VertAccessor.GetVertexUV(LODData, Idx1, UVSet);
-				const FVector2f UV2 = VertAccessor.GetVertexUV(LODData, Idx2, UVSet);
+				const FVector2f UV0 = FVector2f(VertAccessor.GetVertexUV(LODData, Idx0, UVSet));	// LWC_TODO: Precision loss
+				const FVector2f UV1 = FVector2f(VertAccessor.GetVertexUV(LODData, Idx1, UVSet));	// LWC_TODO: Precision loss
+				const FVector2f UV2 = FVector2f(VertAccessor.GetVertexUV(LODData, Idx2, UVSet));	// LWC_TODO: Precision loss
 
 				FVector2f UV = BarycentricInterpolate(BaryParam.GetAndAdvance(), UV0, UV1, UV2);
 				OutUV.SetAndAdvance(UV);
@@ -1083,7 +1083,7 @@ void UNiagaraDataInterfaceSkeletalMesh::GetTriangleCoordAtUV(FVectorVMExternalFu
 		for (int32 i = 0; i < Context.GetNumInstances(); ++i)
 		{
 			const bool Enabled = InEnabled.GetAndAdvance();
-			const FVector2D SourceUv = InUV.GetAndAdvance();
+			const FVector2D SourceUv = FVector2D(InUV.GetAndAdvance());
 			const float Tolerance = InTolerance.GetAndAdvance();
 
 			FVector3f BaryCoord(ForceInitToZero);
@@ -1153,8 +1153,8 @@ void UNiagaraDataInterfaceSkeletalMesh::GetTriangleCoordInAabb(FVectorVMExternal
 		for (int32 i = 0; i < Context.GetNumInstances(); ++i)
 		{
 			const bool Enabled = InEnabled.GetAndAdvance();
-			const FVector2D MinExtent = InMinExtent.GetAndAdvance();
-			const FVector2D MaxExtent = InMaxExtent.GetAndAdvance();
+			const FVector2D MinExtent = FVector2D(InMinExtent.GetAndAdvance());
+			const FVector2D MaxExtent = FVector2D(InMaxExtent.GetAndAdvance());
 
 			FVector3f BaryCoord(ForceInitToZero);
 			int32 TriangleIndex = INDEX_NONE;
diff --git a/Engine/Plugins/FX/Niagara/Source/Niagara/Private/NDISkeletalMesh_VertexSampling.cpp b/Engine/Plugins/FX/Niagara/Source/Niagara/Private/NDISkeletalMesh_VertexSampling.cpp
--- a/Engine/Plugins/FX/Niagara/Source/Niagara/Private/NDISkeletalMesh_VertexSampling.cpp
+++ b/Engine/Plugins/FX/Niagara/Source/Niagara/Private/NDISkeletalMesh_VertexSampling.cpp
@@ -642,7 +642,7 @@ void UNiagaraDataInterfaceSkeletalMesh::GetVertexUV(FVectorVMExternalFunctionCon
 				const int32 Vert = FMath::Clamp(VertParam.GetAndAdvance(), 0, VertMax);
 				const int32 UVSet = UVSetParam.GetAndAdvance();
 				const FVector2D UV = VertAccessor.GetVertexUV(LODData, Vert, UVSet);
-				OutUV.SetAndAdvance(UV);
+				OutUV.SetAndAdvance(FVector2f(UV));	// LWC_TODO: Precision loss
 			}
 			// We are done
 			return;
diff --git a/Engine/Plugins/FX/Niagara/Source/Niagara/Private/NiagaraBoundsCalculatorHelper.h b/Engine/Plugins/FX/Niagara/Source/Niagara/Private/NiagaraBoundsCalculatorHelper.h
--- a/Engine/Plugins/FX/Niagara/Source/Niagara/Private/NiagaraBoundsCalculatorHelper.h
+++ b/Engine/Plugins/FX/Niagara/Source/Niagara/Private/NiagaraBoundsCalculatorHelper.h
@@ -91,7 +91,7 @@ public:
 
 			if (SpriteSizeAccessor.IsValid())
 			{
-				const FVector2D MaxSpriteSize2D = SpriteSizeAccessor.GetReader(DataSet).GetMax();
+				const FVector2f MaxSpriteSize2D = SpriteSizeAccessor.GetReader(DataSet).GetMax();
 				MaxSpriteSize = FMath::Max(MaxSpriteSize2D.X, MaxSpriteSize2D.Y);
 			}
 			MaxSize = FMath::Max(MaxSize, FMath::IsNearlyZero(MaxSpriteSize) ? 1.0f : MaxSpriteSize);
diff --git a/Engine/Plugins/FX/Niagara/Source/Niagara/Private/NiagaraDataInterfaceGrid2DCollection.cpp b/Engine/Plugins/FX/Niagara/Source/Niagara/Private/NiagaraDataInterfaceGrid2DCollection.cpp
--- a/Engine/Plugins/FX/Niagara/Source/Niagara/Private/NiagaraDataInterfaceGrid2DCollection.cpp
+++ b/Engine/Plugins/FX/Niagara/Source/Niagara/Private/NiagaraDataInterfaceGrid2DCollection.cpp
@@ -2293,10 +2293,11 @@ void UNiagaraDataInterfaceGrid2DCollection::GetWorldBBoxSize(FVectorVMExternalFu
 {
 	VectorVM::FUserPtrHandler<FGrid2DCollectionRWInstanceData_GameThread> InstData(Context);
 	FNDIOutputParam<FVector2f> OutWorldBounds(Context);
+	const FVector2f BBoxSize = FVector2f(InstData->WorldBBoxSize);	// LWC_TODO: Precision loss
 
 	for (int32 InstanceIdx = 0; InstanceIdx < Context.GetNumInstances(); ++InstanceIdx)
 	{
-		OutWorldBounds.SetAndAdvance(InstData->WorldBBoxSize);
+		OutWorldBounds.SetAndAdvance(BBoxSize);
 	}
 }
 
@@ -2305,10 +2306,11 @@ void UNiagaraDataInterfaceGrid2DCollection::GetCellSize(FVectorVMExternalFunctio
 {
 	VectorVM::FUserPtrHandler<FGrid2DCollectionRWInstanceData_GameThread> InstData(Context);
 	FNDIOutputParam<FVector2f> OutCellSize(Context);
+	const FVector2f CellSize = FVector2f(InstData->CellSize);	// LWC_TODO: Precision loss
 
 	for (int32 InstanceIdx = 0; InstanceIdx < Context.GetNumInstances(); ++InstanceIdx)
 	{
-		OutCellSize.SetAndAdvance(InstData->CellSize);
+		OutCellSize.SetAndAdvance(CellSize);
 	}
 }
 
diff --git a/Engine/Plugins/FX/Niagara/Source/Niagara/Private/NiagaraRendererMeshes.cpp b/Engine/Plugins/FX/Niagara/Source/Niagara/Private/NiagaraRendererMeshes.cpp
--- a/Engine/Plugins/FX/Niagara/Source/Niagara/Private/NiagaraRendererMeshes.cpp
+++ b/Engine/Plugins/FX/Niagara/Source/Niagara/Private/NiagaraRendererMeshes.cpp
@@ -82,7 +82,7 @@ FNiagaraRendererMeshes::FNiagaraRendererMeshes(ERHIFeatureLevel::Type FeatureLev
 	bSortOnlyWhenTranslucent = Properties->bSortOnlyWhenTranslucent;
 	bGpuLowLatencyTranslucency = Properties->bGpuLowLatencyTranslucency;
 	bOverrideMaterials = Properties->bOverrideMaterials;
-	SubImageSize = Properties->SubImageSize;
+	SubImageSize = FVector2f(Properties->SubImageSize);	// LWC_TODO: Precision loss
 	bSubImageBlend = Properties->bSubImageBlend;
 	bEnableFrustumCulling = Properties->bEnableFrustumCulling;
 	bEnableCulling = bEnableFrustumCulling;
diff --git a/Engine/Plugins/FX/Niagara/Source/Niagara/Private/NiagaraRendererSprites.cpp b/Engine/Plugins/FX/Niagara/Source/Niagara/Private/NiagaraRendererSprites.cpp
--- a/Engine/Plugins/FX/Niagara/Source/Niagara/Private/NiagaraRendererSprites.cpp
+++ b/Engine/Plugins/FX/Niagara/Source/Niagara/Private/NiagaraRendererSprites.cpp
@@ -90,10 +90,10 @@ FNiagaraRendererSprites::FNiagaraRendererSprites(ERHIFeatureLevel::Type FeatureL
 	SourceMode = Properties->SourceMode;
 	Alignment = Properties->Alignment;
 	FacingMode = Properties->FacingMode;
-	PivotInUVSpace = Properties->PivotInUVSpace;
+	PivotInUVSpace = FVector2f(Properties->PivotInUVSpace);	// LWC_TODO: Precision loss
 	MacroUVRadius = Properties->MacroUVRadius;
 	SortMode = Properties->SortMode;
-	SubImageSize = Properties->SubImageSize;
+	SubImageSize = FVector2f(Properties->SubImageSize);	// LWC_TODO: Precision loss
 	NumIndicesPerInstance = Properties->GetNumIndicesPerInstance();
 	bSubImageBlend = Properties->bSubImageBlend;
 	bRemoveHMDRollInVR = Properties->bRemoveHMDRollInVR;
diff --git a/Engine/Plugins/FX/Niagara/Source/Niagara/Private/NiagaraTypes.cpp b/Engine/Plugins/FX/Niagara/Source/Niagara/Private/NiagaraTypes.cpp
--- a/Engine/Plugins/FX/Niagara/Source/Niagara/Private/NiagaraTypes.cpp
+++ b/Engine/Plugins/FX/Niagara/Source/Niagara/Private/NiagaraTypes.cpp
@@ -74,19 +74,19 @@ void FNiagaraStructConversionStep::CopyToSim(uint8* DestinationData, const uint8
 	else if (ConversionType == ENiagaraStructConversionType::Vector2)
 	{
 		checkf(SourceBytes == 16 && SimulationBytes == 8, TEXT("Wrong bytesizes for Vector2 d->f conversion: Source %i, Simulation %i"), SourceBytes, SimulationBytes);
-		FVector2f Vector2 = *static_cast<const FVector2d*>(Src);
+		FVector2f Vector2(*static_cast<const FVector2d*>(Src));
 		FMemory::Memcpy(Dest, &Vector2, SimulationBytes);
 	}
 	else if (ConversionType == ENiagaraStructConversionType::Vector3)
 	{
 		checkf(SourceBytes == 24 && SimulationBytes == 12, TEXT("Wrong bytesizes for Vector3 d->f conversion: Source %i, Simulation %i"), SourceBytes, SimulationBytes);
-		FVector3f Vector = *static_cast<const FVector3d*>(Src);
+		FVector3f Vector(*static_cast<const FVector3d*>(Src));
 		FMemory::Memcpy(Dest, &Vector, SimulationBytes);
 	}
 	else if (ConversionType == ENiagaraStructConversionType::Vector4)
 	{
 		checkf(SourceBytes == 32 && SimulationBytes == 16, TEXT("Wrong bytesizes for Vector4 d->f conversion: Source %i, Simulation %i"), SourceBytes, SimulationBytes);
-		FVector4f Vector4 = FVector4f(*static_cast<const FVector4d*>(Src));
+		FVector4f Vector4(*static_cast<const FVector4d*>(Src));
 		FMemory::Memcpy(Dest, &Vector4, SimulationBytes);
 	}
 	else if (ConversionType == ENiagaraStructConversionType::Quat)
@@ -115,19 +115,19 @@ void FNiagaraStructConversionStep::CopyFromSim(uint8* DestinationData, const uin
 	else if (ConversionType == ENiagaraStructConversionType::Vector2)
 	{
 		checkf(SourceBytes == 16 && SimulationBytes == 8, TEXT("Wrong bytesizes for Vector2 f->d conversion: Source %i, Simulation %i"), SourceBytes, SimulationBytes);
-		FVector2d Vector2 = *static_cast<const FVector2f*>(Src);
+		FVector2d Vector2(*static_cast<const FVector2f*>(Src));
 		FMemory::Memcpy(Dest, &Vector2, SourceBytes);
 	}
 	else if (ConversionType == ENiagaraStructConversionType::Vector3)
 	{
 		checkf(SourceBytes == 24 && SimulationBytes == 12, TEXT("Wrong bytesizes for Vector3 f->d conversion: Source %i, Simulation %i"), SourceBytes, SimulationBytes);
-		FVector3d Vector3 = *static_cast<const FVector3f*>(Src);
+		FVector3d Vector3(*static_cast<const FVector3f*>(Src));
         FMemory::Memcpy(Dest, &Vector3, SourceBytes);
 	}
 	else if (ConversionType == ENiagaraStructConversionType::Vector4)
 	{
 		checkf(SourceBytes == 32 && SimulationBytes == 16, TEXT("Wrong bytesizes for Vector4 f->d conversion: Source %i, Simulation %i"), SourceBytes, SimulationBytes);
-		FVector4d Vector4 = FVector4d(*static_cast<const FVector4f*>(Src));
+		FVector4d Vector4(*static_cast<const FVector4f*>(Src));
 		FMemory::Memcpy(Dest, &Vector4, SourceBytes);
 	}
 	else if (ConversionType == ENiagaraStructConversionType::Quat)
diff --git a/Engine/Plugins/FX/Niagara/Source/NiagaraEditor/Private/NiagaraClipboard.cpp b/Engine/Plugins/FX/Niagara/Source/NiagaraEditor/Private/NiagaraClipboard.cpp
--- a/Engine/Plugins/FX/Niagara/Source/NiagaraEditor/Private/NiagaraClipboard.cpp
+++ b/Engine/Plugins/FX/Niagara/Source/NiagaraEditor/Private/NiagaraClipboard.cpp
@@ -268,7 +268,7 @@ UNiagaraClipboardFunctionInput* UNiagaraClipboardEditorScriptingUtilities::Creat
 UNiagaraClipboardFunctionInput* UNiagaraClipboardEditorScriptingUtilities::CreateVec2LocalValueInput(UObject* InOuter, FName InInputName, bool bInHasEditCondition, bool bInEditConditionValue, FVector2D InVec2Value)
 {
 	FNiagaraTypeDefinition InputType = FNiagaraTypeDefinition::GetVec2Def();
-	const FVector2f InVec2ValueFloat = InVec2Value;
+	const FVector2f InVec2ValueFloat = FVector2f(InVec2Value);	// LWC_TODO: Precision loss
 	TArray<uint8> Vec2ValueFloat;
 	Vec2ValueFloat.AddUninitialized(InputType.GetSize());
 	FMemory::Memcpy(Vec2ValueFloat.GetData(), &InVec2ValueFloat, InputType.GetSize());
diff --git a/Engine/Plugins/FX/Niagara/Source/NiagaraShader/Private/NiagaraGenerateMips.cpp b/Engine/Plugins/FX/Niagara/Source/NiagaraShader/Private/NiagaraGenerateMips.cpp
--- a/Engine/Plugins/FX/Niagara/Source/NiagaraShader/Private/NiagaraGenerateMips.cpp
+++ b/Engine/Plugins/FX/Niagara/Source/NiagaraShader/Private/NiagaraGenerateMips.cpp
@@ -74,8 +74,8 @@ void NiagaraGenerateMips::GenerateMips(FRHICommandList& RHICmdList, FRHITexture2
 		FUnorderedAccessViewRHIRef MipOutUAV = RHICreateUnorderedAccessView(TextureRHI, iDstMip);
 
 		FNiagaraGenerateMipsCS::FParameters PassParameters;
-		PassParameters.SrcTexelSize	= FVector2D(1.0f / float(SrcMipSize.X), 1.0f / float(SrcMipSize.Y));
-		PassParameters.DstTexelSize = FVector2D(1.0f / float(DstMipSize.X), 1.0f / float(DstMipSize.Y));
+		PassParameters.SrcTexelSize	= FVector2f(1.0f / float(SrcMipSize.X), 1.0f / float(SrcMipSize.Y));
+		PassParameters.DstTexelSize = FVector2f(1.0f / float(DstMipSize.X), 1.0f / float(DstMipSize.Y));
 		PassParameters.KernelHWidth = KernelHWidth;
 		PassParameters.MipOutSize	= DstMipSize;
 		PassParameters.MipInSRV		= MipInSRV;
diff --git a/Engine/Plugins/Importers/AlembicImporter/Source/AlembicLibrary/Private/AbcImportUtilities.cpp b/Engine/Plugins/Importers/AlembicImporter/Source/AlembicLibrary/Private/AbcImportUtilities.cpp
--- a/Engine/Plugins/Importers/AlembicImporter/Source/AlembicLibrary/Private/AbcImportUtilities.cpp
+++ b/Engine/Plugins/Importers/AlembicImporter/Source/AlembicLibrary/Private/AbcImportUtilities.cpp
@@ -1318,13 +1318,13 @@ void AbcImporterUtilities::ApplyConversion(FAbcMeshSample* InOutSample, const FA
 	{
 		// Apply UV matrix to flip channels
 		FMatrix2x2 UVMatrix( FScale2D(InConversionSettings.bFlipU ? -1.0f : 1.0f, InConversionSettings.bFlipV ? -1.0f : 1.0f) );
-		FVector2D UVOffset(InConversionSettings.bFlipU ? 1.0f : 0.0f, InConversionSettings.bFlipV ? 1.0f : 0.0f);
+		FVector2f UVOffset(InConversionSettings.bFlipU ? 1.0f : 0.0f, InConversionSettings.bFlipV ? 1.0f : 0.0f);
 				
 		for (uint32 UVIndex = 0; UVIndex < InOutSample->NumUVSets; ++UVIndex)
 		{
 			for (FVector2f& UV : InOutSample->UVs[UVIndex])
 			{
-				UV = UVOffset + UVMatrix.TransformPoint(UV);
+				UV = UVOffset + FVector2f(UVMatrix.TransformPoint(FVector2D(UV)));
 			}
 		}
 	}
diff --git a/Engine/Plugins/Importers/AlembicImporter/Source/AlembicLibrary/Private/AbcImporter.cpp b/Engine/Plugins/Importers/AlembicImporter/Source/AlembicLibrary/Private/AbcImporter.cpp
--- a/Engine/Plugins/Importers/AlembicImporter/Source/AlembicLibrary/Private/AbcImporter.cpp
+++ b/Engine/Plugins/Importers/AlembicImporter/Source/AlembicLibrary/Private/AbcImporter.cpp
@@ -1613,7 +1613,7 @@ bool FAbcImporter::BuildSkeletalMesh( FSkeletalMeshLODModel& LODModel, const FRe
 
 			for (uint32 UVIndex = 0; UVIndex < Sample->NumUVSets; ++UVIndex)
 			{
-				Section.UVs[UVIndex].Add(Sample->UVs[UVIndex][FaceOffset + VertexIndex]);
+				Section.UVs[UVIndex].Add(FVector2D(Sample->UVs[UVIndex][FaceOffset + VertexIndex]));
 			}		
 			
 			Section.Colors.Add(Sample->Colors[FaceOffset + VertexIndex].ToFColor(false));
@@ -1675,7 +1675,7 @@ bool FAbcImporter::BuildSkeletalMesh( FSkeletalMeshLODModel& LODModel, const FRe
 				NewVertex.TangentZ = (FVector3f)SourceSection.TangentZ[FaceOffset + VertexIndex]; // LWC_TODO: precision loss
 				for (uint32 UVIndex = 0; UVIndex < SourceSection.NumUVSets; ++UVIndex)
 				{
-					NewVertex.UVs[UVIndex] = SourceSection.UVs[UVIndex][FaceOffset + VertexIndex];
+					NewVertex.UVs[UVIndex] = FVector2f(SourceSection.UVs[UVIndex][FaceOffset + VertexIndex]);
 				}
 				
 				NewVertex.Color = SourceSection.Colors[FaceOffset + VertexIndex];
diff --git a/Engine/Plugins/Importers/USDImporter/Source/USDUtilities/Private/USDGeomMeshConversion.cpp b/Engine/Plugins/Importers/USDImporter/Source/USDUtilities/Private/USDGeomMeshConversion.cpp
--- a/Engine/Plugins/Importers/USDImporter/Source/USDUtilities/Private/USDGeomMeshConversion.cpp
+++ b/Engine/Plugins/Importers/USDImporter/Source/USDUtilities/Private/USDGeomMeshConversion.cpp
@@ -179,7 +179,7 @@ namespace UE
 
 								for ( int32 VertexIndex = 0; VertexIndex < VertexCount; ++VertexIndex )
 								{
-									FVector2D TexCoord = LODRenderMesh.VertexBuffers.StaticMeshVertexBuffer.GetVertexUV( VertexIndex, TexCoordSourceIndex );
+									FVector2D TexCoord = FVector2D(LODRenderMesh.VertexBuffers.StaticMeshVertexBuffer.GetVertexUV( VertexIndex, TexCoordSourceIndex ));
 									TexCoord[ 1 ] = 1.f - TexCoord[ 1 ];
 
 									UVs.push_back( UnrealToUsd::ConvertVector( TexCoord ) );
@@ -436,7 +436,7 @@ namespace UE
 
 							for ( const FVertexInstanceID InstanceID : MeshDescription.VertexInstances().GetElementIDs() )
 							{
-								FVector2D UV = VertexInstanceUVs.Get( InstanceID, UVIndex );
+								FVector2D UV = FVector2D(VertexInstanceUVs.Get( InstanceID, UVIndex ));
 								UV[ 1 ] = 1.f - UV[ 1 ];
 								UVs.push_back( UnrealToUsd::ConvertVector( UV ) );
 							}
@@ -913,7 +913,7 @@ bool UsdToUnreal::ConvertGeomMesh( const pxr::UsdTyped& UsdSchema, FMeshDescript
 					}
 
 					// Flip V for Unreal uv's which match directx
-					FVector2D FinalUVVector( UV[0], 1.f - UV[1] );
+					FVector2f FinalUVVector( UV[0], 1.f - UV[1] );
 					MeshDescriptionUVs.Set( AddedVertexInstanceId, UVSet.UVSetIndexUE, FinalUVVector );
 				}
 
diff --git a/Engine/Plugins/Importers/USDImporter/Source/USDUtilities/Private/USDSkeletalDataConversion.cpp b/Engine/Plugins/Importers/USDImporter/Source/USDUtilities/Private/USDSkeletalDataConversion.cpp
--- a/Engine/Plugins/Importers/USDImporter/Source/USDUtilities/Private/USDSkeletalDataConversion.cpp
+++ b/Engine/Plugins/Importers/USDImporter/Source/USDUtilities/Private/USDSkeletalDataConversion.cpp
@@ -735,7 +735,7 @@ namespace UnrealToUsdImpl
 
 						for ( int32 VertexIndex = 0; VertexIndex < VertexCount; ++VertexIndex )
 						{
-							FVector2D TexCoord = Vertices[ VertexIndex ].UVs[ TexCoordSourceIndex ];
+							FVector2D TexCoord = FVector2D(Vertices[ VertexIndex ].UVs[ TexCoordSourceIndex ]);
 							TexCoord[ 1 ] = 1.f - TexCoord[ 1 ];
 
 							UVs.push_back( UnrealToUsd::ConvertVector( TexCoord ) );
@@ -1552,7 +1552,7 @@ bool UsdToUnreal::ConvertSkinnedMesh(const pxr::UsdSkelSkinningQuery& SkinningQu
 					}
 
 					// Flip V for Unreal uv's which match directx
-					FVector2D FinalUVVector( UV[0], 1.f - UV[1] );
+					FVector2f FinalUVVector( UV[0], 1.f - UV[1] );
 					SkelMeshWedge.UVs[ UVLayerIndex ] = FinalUVVector;
 
 					++UVLayerIndex;
diff --git a/Engine/Plugins/MeshPainting/Source/MeshPaintEditorMode/Private/MeshPaintModeHelpers.cpp b/Engine/Plugins/MeshPainting/Source/MeshPaintEditorMode/Private/MeshPaintModeHelpers.cpp
--- a/Engine/Plugins/MeshPainting/Source/MeshPaintEditorMode/Private/MeshPaintModeHelpers.cpp
+++ b/Engine/Plugins/MeshPainting/Source/MeshPaintEditorMode/Private/MeshPaintModeHelpers.cpp
@@ -280,7 +280,7 @@ void UMeshPaintModeSubsystem::ImportVertexColorsToSkeletalMesh(USkeletalMesh* Sk
 
 		for (uint32 VertexIndex = 0; VertexIndex < LODData.GetNumVertices(); ++VertexIndex)
 		{
-			const FVector2D UV = LODData.StaticVertexBuffers.StaticMeshVertexBuffer.GetVertexUV(VertexIndex, UVIndex);
+			const FVector2D UV = FVector2D(LODData.StaticVertexBuffers.StaticMeshVertexBuffer.GetVertexUV(VertexIndex, UVIndex));
 			LODData.StaticVertexBuffers.ColorVertexBuffer.VertexColor(VertexIndex) = GEngine->GetEngineSubsystem<UMeshPaintingSubsystem>()->PickVertexColorFromTextureData(MipData, UV, Texture, ColorMask);
 		}
 
@@ -297,7 +297,7 @@ void UMeshPaintModeSubsystem::ImportVertexColorsToSkeletalMesh(USkeletalMesh* Sk
 		int32 SectionVertexIndex = INDEX_NONE;
 		LODModel.GetSectionFromVertexIndex(VertexIndex, SectionIndex, SectionVertexIndex);
 
-		const FVector2D UV = LODModel.Sections[SectionIndex].SoftVertices[SectionVertexIndex].UVs[UVIndex];
+		const FVector2D UV = FVector2D(LODModel.Sections[SectionIndex].SoftVertices[SectionVertexIndex].UVs[UVIndex]);
 		LODModel.Sections[SectionIndex].SoftVertices[SectionVertexIndex].Color = GEngine->GetEngineSubsystem<UMeshPaintingSubsystem>()->PickVertexColorFromTextureData(MipData, UV, Texture, ColorMask);
 	}
 
@@ -391,7 +391,7 @@ void UMeshPaintModeSubsystem::ImportVertexColorsToStaticMesh(UStaticMesh* Static
 	const FColor ColorMask = Options->CreateColorMask();
 	for (uint32 VertexIndex = 0; VertexIndex < LODModel.VertexBuffers.StaticMeshVertexBuffer.GetNumVertices(); ++VertexIndex)
 	{
-		const FVector2D UV = LODModel.VertexBuffers.StaticMeshVertexBuffer.GetVertexUV(VertexIndex, UVIndex);
+		const FVector2D UV = FVector2D(LODModel.VertexBuffers.StaticMeshVertexBuffer.GetVertexUV(VertexIndex, UVIndex));
 		LODModel.VertexBuffers.ColorVertexBuffer.VertexColor(VertexIndex) = GEngine->GetEngineSubsystem<UMeshPaintingSubsystem>()->PickVertexColorFromTextureData(MipData, UV, Texture, ColorMask);
 	}
 
@@ -454,7 +454,7 @@ void UMeshPaintModeSubsystem::ImportVertexColorsToStaticMeshComponent(UStaticMes
 		const FColor ColorMask = Options->CreateColorMask();
 		for (uint32 VertexIndex = 0; VertexIndex < LODModel.VertexBuffers.StaticMeshVertexBuffer.GetNumVertices(); ++VertexIndex)
 		{
-			const FVector2D UV = LODModel.VertexBuffers.StaticMeshVertexBuffer.GetVertexUV(VertexIndex, UVIndex);
+			const FVector2D UV = FVector2D(LODModel.VertexBuffers.StaticMeshVertexBuffer.GetVertexUV(VertexIndex, UVIndex));
 			InstanceMeshLODInfo.OverrideVertexColors->VertexColor(VertexIndex) = GEngine->GetEngineSubsystem<UMeshPaintingSubsystem>()->PickVertexColorFromTextureData(MipData, UV, Texture, ColorMask);
 		}
 
diff --git a/Engine/Plugins/MeshPainting/Source/MeshPaintingToolset/Private/MeshPaintSkeletalMeshAdapter.cpp b/Engine/Plugins/MeshPainting/Source/MeshPaintingToolset/Private/MeshPaintSkeletalMeshAdapter.cpp
--- a/Engine/Plugins/MeshPainting/Source/MeshPaintingToolset/Private/MeshPaintSkeletalMeshAdapter.cpp
+++ b/Engine/Plugins/MeshPainting/Source/MeshPaintingToolset/Private/MeshPaintSkeletalMeshAdapter.cpp
@@ -320,7 +320,7 @@ void FMeshPaintSkeletalMeshComponentAdapter::ApplyOrRemoveTextureOverride(UTextu
 
 void FMeshPaintSkeletalMeshComponentAdapter::GetTextureCoordinate(int32 VertexIndex, int32 ChannelIndex, FVector2D& OutTextureCoordinate) const
 {
-	OutTextureCoordinate = LODData->StaticVertexBuffers.StaticMeshVertexBuffer.GetVertexUV(VertexIndex, ChannelIndex);
+	OutTextureCoordinate = FVector2D(LODData->StaticVertexBuffers.StaticMeshVertexBuffer.GetVertexUV(VertexIndex, ChannelIndex));
 }
 
 void FMeshPaintSkeletalMeshComponentAdapter::PreEdit()
diff --git a/Engine/Plugins/MeshPainting/Source/MeshPaintingToolset/Private/MeshPaintStaticMeshAdapter.cpp b/Engine/Plugins/MeshPainting/Source/MeshPaintingToolset/Private/MeshPaintStaticMeshAdapter.cpp
--- a/Engine/Plugins/MeshPainting/Source/MeshPaintingToolset/Private/MeshPaintStaticMeshAdapter.cpp
+++ b/Engine/Plugins/MeshPainting/Source/MeshPaintingToolset/Private/MeshPaintStaticMeshAdapter.cpp
@@ -336,7 +336,7 @@ FMatrix FMeshPaintStaticMeshComponentAdapter::GetComponentToWorldMatrix() const
 
 void FMeshPaintStaticMeshComponentAdapter::GetTextureCoordinate(int32 VertexIndex, int32 ChannelIndex, FVector2D& OutTextureCoordinate) const
 {
-	OutTextureCoordinate = LODModel->VertexBuffers.StaticMeshVertexBuffer.GetVertexUV(VertexIndex, ChannelIndex);
+	OutTextureCoordinate = FVector2D(LODModel->VertexBuffers.StaticMeshVertexBuffer.GetVertexUV(VertexIndex, ChannelIndex));
 }
 
 
diff --git a/Engine/Plugins/MeshPainting/Source/MeshPaintingToolset/Private/TexturePaintToolset.cpp b/Engine/Plugins/MeshPainting/Source/MeshPaintingToolset/Private/TexturePaintToolset.cpp
--- a/Engine/Plugins/MeshPainting/Source/MeshPaintingToolset/Private/TexturePaintToolset.cpp
+++ b/Engine/Plugins/MeshPainting/Source/MeshPaintingToolset/Private/TexturePaintToolset.cpp
@@ -224,7 +224,7 @@ bool UTexturePaintToolset::GenerateSeamMask(UMeshComponent* MeshComponent, int32
 			for (int32 TriVertexNum = 0; TriVertexNum < 3; ++TriVertexNum)
 			{
 				const int32 VertexIndex = Indices[TriIndex * 3 + TriVertexNum];
-				TriUVs[TriVertexNum] = LODModel.VertexBuffers.StaticMeshVertexBuffer.GetVertexUV(VertexIndex, UVSet);
+				TriUVs[TriVertexNum] = FVector2D(LODModel.VertexBuffers.StaticMeshVertexBuffer.GetVertexUV(VertexIndex, UVSet));
 
 				// Update bounds
 				float U = TriUVs[TriVertexNum].X;
diff --git a/Engine/Plugins/Runtime/AR/AppleAR/AppleARKit/Source/AppleARKit/Private/AppleARKitVideoOverlay.cpp b/Engine/Plugins/Runtime/AR/AppleAR/AppleARKit/Source/AppleARKit/Private/AppleARKitVideoOverlay.cpp
--- a/Engine/Plugins/Runtime/AR/AppleAR/AppleARKit/Source/AppleARKit/Private/AppleARKitVideoOverlay.cpp
+++ b/Engine/Plugins/Runtime/AR/AppleAR/AppleARKit/Source/AppleARKit/Private/AppleARKitVideoOverlay.cpp
@@ -298,7 +298,7 @@ void FAppleARKitVideoOverlay::RenderVideoOverlayWithMaterial(FRHICommandListImme
 		{
 			const auto& Position = Positions[Index];
 			Vertices[Index].Position = Position;
-			Vertices[Index].UV = FVector2D(Position.X, Position.Y);
+			Vertices[Index].UV = FVector2f(Position.X, Position.Y);
 		}
 		
 		FRHIResourceCreateInfo CreateInfoVB(TEXT("VideoOverlayVertexBuffer"), &Vertices);
diff --git a/Engine/Plugins/Runtime/CableComponent/Source/CableComponent/Private/CableComponent.cpp b/Engine/Plugins/Runtime/CableComponent/Source/CableComponent/Private/CableComponent.cpp
--- a/Engine/Plugins/Runtime/CableComponent/Source/CableComponent/Private/CableComponent.cpp
+++ b/Engine/Plugins/Runtime/CableComponent/Source/CableComponent/Private/CableComponent.cpp
@@ -148,7 +148,7 @@ public:
 
 				FDynamicMeshVertex Vert;
 				Vert.Position = InPoints[PointIdx] + (OutDir * 0.5f * CableWidth);
-				Vert.TextureCoordinate[0] = FVector2D(AlongFrac * TileMaterial, AroundFrac);
+				Vert.TextureCoordinate[0] = FVector2f(AlongFrac * TileMaterial, AroundFrac);
 				Vert.Color = VertexColor;
 				Vert.SetTangents(ForwardDir, OutDir ^ ForwardDir, OutDir);
 				OutVertices.Add(Vert);
diff --git a/Engine/Plugins/Runtime/GeometryCache/Source/GeometryCache/Private/CodecV1Test.h b/Engine/Plugins/Runtime/GeometryCache/Source/GeometryCache/Private/CodecV1Test.h
--- a/Engine/Plugins/Runtime/GeometryCache/Source/GeometryCache/Private/CodecV1Test.h
+++ b/Engine/Plugins/Runtime/GeometryCache/Source/GeometryCache/Private/CodecV1Test.h
@@ -122,8 +122,8 @@ public:
 						check(true);
 					}
 
-					const FVector2D& UVA = OriginalMeshData.TextureCoordinates[Index];
-					const FVector2D& UVB = DecodedMeshData.TextureCoordinates[Index];
+					const FVector2f& UVA = OriginalMeshData.TextureCoordinates[Index];
+					const FVector2f& UVB = DecodedMeshData.TextureCoordinates[Index];
 
 					if (!UVsEqual(UVA, UVB))
 					{
diff --git a/Engine/Plugins/Runtime/GeometryCache/Source/GeometryCache/Private/GeometryCachePreprocessor.cpp b/Engine/Plugins/Runtime/GeometryCache/Source/GeometryCache/Private/GeometryCachePreprocessor.cpp
--- a/Engine/Plugins/Runtime/GeometryCache/Source/GeometryCache/Private/GeometryCachePreprocessor.cpp
+++ b/Engine/Plugins/Runtime/GeometryCache/Source/GeometryCache/Private/GeometryCachePreprocessor.cpp
@@ -121,8 +121,8 @@ bool FOptimizeGeometryCachePreprocessor::AreIndexedVerticesEqual(int32 IndexBuff
 
 		if (BufferedFrames[Frame].MeshData.Positions.Num() == BufferedFrames[Frame].MeshData.TextureCoordinates.Num())
 		{
-			const FVector2D& UVA = BufferedFrames[Frame].MeshData.TextureCoordinates[VertexIndexA];
-			const FVector2D& UVB = BufferedFrames[Frame].MeshData.TextureCoordinates[VertexIndexB];
+			const FVector2f& UVA = BufferedFrames[Frame].MeshData.TextureCoordinates[VertexIndexA];
+			const FVector2f& UVB = BufferedFrames[Frame].MeshData.TextureCoordinates[VertexIndexB];
 
 			if (!UVsEqual(UVA, UVB))
 			{
diff --git a/Engine/Plugins/Runtime/HairStrands/Source/HairStrandsCore/Private/GroomBindingBuilder.cpp b/Engine/Plugins/Runtime/HairStrands/Source/HairStrandsCore/Private/GroomBindingBuilder.cpp
--- a/Engine/Plugins/Runtime/HairStrands/Source/HairStrandsCore/Private/GroomBindingBuilder.cpp
+++ b/Engine/Plugins/Runtime/HairStrands/Source/HairStrandsCore/Private/GroomBindingBuilder.cpp
@@ -226,7 +226,7 @@ public:
 
 	virtual FVector2D GetVertexUV(uint32 VertexIndex, uint32 ChannelIndex) const override
 	{
-		return MeshLODData.StaticVertexBuffers.StaticMeshVertexBuffer.GetVertexUV(VertexIndex, ChannelIndex);
+		return FVector2D(MeshLODData.StaticVertexBuffers.StaticMeshVertexBuffer.GetVertexUV(VertexIndex, ChannelIndex));
 	}
 
 	virtual void GetSectionFromVertexIndex(uint32 InVertIndex, int32& OutSectionIndex) const override
@@ -415,7 +415,7 @@ public:
 
 	virtual FVector2D GetVertexUV(uint32 VertexIndex, uint32 ChannelIndex) const override
 	{
-		return MeshData.TextureCoordinates[VertexIndex];
+		return FVector2D(MeshData.TextureCoordinates[VertexIndex]);
 	}
 
 	virtual void GetSectionFromVertexIndex(uint32 InVertIndex, int32& OutSectionIndex) const override
@@ -1281,13 +1281,13 @@ namespace GroomBinding_RootProjection
 				}
 				check(ClosestDistance < FLT_MAX);
 
-				const uint32 EncodedBarycentrics = FHairStrandsRootUtils::EncodeBarycentrics(ClosestBarycentrics);
+				const uint32 EncodedBarycentrics = FHairStrandsRootUtils::EncodeBarycentrics(FVector2f(ClosestBarycentrics));	// LWC_TODO: Precision loss
 				const uint32 EncodedTriangleIndex = FHairStrandsRootUtils::EncodeTriangleIndex(ClosestTriangle.TriangleIndex, ClosestTriangle.SectionIndex);
 				OutRootData.MeshProjectionLODs[LODIt].RootTriangleIndexBuffer[CurveIndex] = EncodedTriangleIndex;
 				OutRootData.MeshProjectionLODs[LODIt].RootTriangleBarycentricBuffer[CurveIndex] = EncodedBarycentrics;
-				OutRootData.MeshProjectionLODs[LODIt].RestRootTrianglePosition0Buffer[CurveIndex] = FVector4f(ClosestTriangle.P0, FHairStrandsRootUtils::PackUVsToFloat(ClosestTriangle.UV0));
-				OutRootData.MeshProjectionLODs[LODIt].RestRootTrianglePosition1Buffer[CurveIndex] = FVector4f(ClosestTriangle.P1, FHairStrandsRootUtils::PackUVsToFloat(ClosestTriangle.UV1));
-				OutRootData.MeshProjectionLODs[LODIt].RestRootTrianglePosition2Buffer[CurveIndex] = FVector4f(ClosestTriangle.P2, FHairStrandsRootUtils::PackUVsToFloat(ClosestTriangle.UV2));
+				OutRootData.MeshProjectionLODs[LODIt].RestRootTrianglePosition0Buffer[CurveIndex] = FVector4f(ClosestTriangle.P0, FHairStrandsRootUtils::PackUVsToFloat(FVector2f(ClosestTriangle.UV0)));	// LWC_TODO: Precision loss
+				OutRootData.MeshProjectionLODs[LODIt].RestRootTrianglePosition1Buffer[CurveIndex] = FVector4f(ClosestTriangle.P1, FHairStrandsRootUtils::PackUVsToFloat(FVector2f(ClosestTriangle.UV1)));	// LWC_TODO: Precision loss
+				OutRootData.MeshProjectionLODs[LODIt].RestRootTrianglePosition2Buffer[CurveIndex] = FVector4f(ClosestTriangle.P2, FHairStrandsRootUtils::PackUVsToFloat(FVector2f(ClosestTriangle.UV2)));	// LWC_TODO: Precision loss
 			}
 #if BINDING_PARALLEL_BUILDING
 			);
diff --git a/Engine/Plugins/Runtime/HairStrands/Source/HairStrandsCore/Private/GroomBuilder.cpp b/Engine/Plugins/Runtime/HairStrands/Source/HairStrandsCore/Private/GroomBuilder.cpp
--- a/Engine/Plugins/Runtime/HairStrands/Source/HairStrandsCore/Private/GroomBuilder.cpp
+++ b/Engine/Plugins/Runtime/HairStrands/Source/HairStrandsCore/Private/GroomBuilder.cpp
@@ -127,7 +127,7 @@ namespace HairStrandsBuilder
 		uint32 Index = 0;
 		for (FVector2f& RootUV : Curves.CurvesRootUV)
 		{
-			RootUV = (RootUVs[Index++] + UVOffset) * UVScale;
+			RootUV = (FVector2f(RootUVs[Index++]) + UVOffset) * UVScale;	// LWC_TODO: Precision loss
 		}
 	}
 
@@ -305,7 +305,7 @@ namespace HairStrandsBuilder
 				PackedPosition.PackedRadiusAndType = PackedRadiusAndType.Packed;
 				PackedPosition.UCoord = uint8(FMath::Clamp(CoordU * 255.f, 0.f, 255.f));
 
-				FVector2D RootUV = Curves.CurvesRootUV[CurveIndex]; 
+				FVector2D RootUV = FVector2D(Curves.CurvesRootUV[CurveIndex]); 
 				FHairStrandsAttribute0Format::Type& PackedAttributes0 = OutPackedAttributes0[PointIndex + IndexOffset];
 				PackedAttributes0.NormalizedLength = uint8(FMath::Clamp(NormalizedLength * 255.f, 0.f, 255.f));
 				PackedAttributes0.Seed = CurveSeed;
diff --git a/Engine/Plugins/Runtime/HairStrands/Source/HairStrandsCore/Private/GroomRBFDeformer.cpp b/Engine/Plugins/Runtime/HairStrands/Source/HairStrandsCore/Private/GroomRBFDeformer.cpp
--- a/Engine/Plugins/Runtime/HairStrands/Source/HairStrandsCore/Private/GroomRBFDeformer.cpp
+++ b/Engine/Plugins/Runtime/HairStrands/Source/HairStrandsCore/Private/GroomRBFDeformer.cpp
@@ -126,7 +126,7 @@ void DeformStrands(
 
 
 		const uint32 PackedBarycentric = RootTriangleBarycentricBuffer[RootIndex];
-		const FVector2D B0 = FHairStrandsRootUtils::DecodeBarycentrics(PackedBarycentric);
+		const FVector2f B0 = FHairStrandsRootUtils::DecodeBarycentrics(PackedBarycentric);
 		const FVector3f   B  = FVector3f(B0.X, B0.Y, 1.f - B0.X - B0.Y);
 
 		/* Strand hair roots translation and rotation in rest position relative to the bound triangle. Positions are relative to the rest root center */
diff --git a/Engine/Plugins/Runtime/HairStrands/Source/HairStrandsCore/Private/GroomTextureBuilder.cpp b/Engine/Plugins/Runtime/HairStrands/Source/HairStrandsCore/Private/GroomTextureBuilder.cpp
--- a/Engine/Plugins/Runtime/HairStrands/Source/HairStrandsCore/Private/GroomTextureBuilder.cpp
+++ b/Engine/Plugins/Runtime/HairStrands/Source/HairStrandsCore/Private/GroomTextureBuilder.cpp
@@ -103,7 +103,7 @@ static void RasterToTexture(int32 Resolution, int32 KernelExtent, uint32 Channel
 	const uint32 CurveCount = InStrandsData.GetNumCurves();
 	for (uint32 CurveIndex = 0; CurveIndex < CurveCount; ++CurveIndex)
 	{
-		FVector2D RootUV = InStrandsData.StrandsCurves.CurvesRootUV[CurveIndex];
+		FVector2f RootUV = InStrandsData.StrandsCurves.CurvesRootUV[CurveIndex];
 		RootUV.Y = FMath::Clamp(1.f - RootUV.Y, 0.f, 1.f);
 
 		const FIntPoint RootCoord(
diff --git a/Engine/Plugins/Runtime/MeshModelingToolset/Source/ModelingComponents/Private/Drawing/LineSetComponent.cpp b/Engine/Plugins/Runtime/MeshModelingToolset/Source/ModelingComponents/Private/Drawing/LineSetComponent.cpp
--- a/Engine/Plugins/Runtime/MeshModelingToolset/Source/ModelingComponents/Private/Drawing/LineSetComponent.cpp
+++ b/Engine/Plugins/Runtime/MeshModelingToolset/Source/ModelingComponents/Private/Drawing/LineSetComponent.cpp
@@ -70,7 +70,7 @@ public:
 			for (const FRenderableLine& OverlayLine : Component->Lines)
 			{
 				const FVector LineDirection = (OverlayLine.End - OverlayLine.Start).GetSafeNormal();
-				const FVector2D UV(OverlayLine.Thickness, OverlayLine.DepthBias);
+				const FVector2f UV(OverlayLine.Thickness, OverlayLine.DepthBias);
 
 				VertexBuffers.PositionVertexBuffer.VertexPosition(VertexBufferIndex + 0) = OverlayLine.Start;
 				VertexBuffers.PositionVertexBuffer.VertexPosition(VertexBufferIndex + 1) = OverlayLine.End;
diff --git a/Engine/Plugins/Runtime/MeshModelingToolset/Source/ModelingComponents/Private/Drawing/MeshWireframeComponent.cpp b/Engine/Plugins/Runtime/MeshModelingToolset/Source/ModelingComponents/Private/Drawing/MeshWireframeComponent.cpp
--- a/Engine/Plugins/Runtime/MeshModelingToolset/Source/ModelingComponents/Private/Drawing/MeshWireframeComponent.cpp
+++ b/Engine/Plugins/Runtime/MeshModelingToolset/Source/ModelingComponents/Private/Drawing/MeshWireframeComponent.cpp
@@ -171,7 +171,7 @@ public:
 			const FVector A = WireSource->GetVertex(EdgeInfo.B);
 			const FVector B = WireSource->GetVertex(EdgeInfo.C);
 			const FVector LineDirection = (B - A).GetSafeNormal();
-			const FVector2D UV(UseThickness, LineDepthBias);
+			const FVector2f UV(UseThickness, LineDepthBias);
 
 			VertexBuffers.PositionVertexBuffer.VertexPosition(VertexBufferIndex + 0) = A;
 			VertexBuffers.PositionVertexBuffer.VertexPosition(VertexBufferIndex + 1) = B;
diff --git a/Engine/Plugins/Runtime/MeshModelingToolset/Source/ModelingComponents/Private/Drawing/PointSetComponent.cpp b/Engine/Plugins/Runtime/MeshModelingToolset/Source/ModelingComponents/Private/Drawing/PointSetComponent.cpp
--- a/Engine/Plugins/Runtime/MeshModelingToolset/Source/ModelingComponents/Private/Drawing/PointSetComponent.cpp
+++ b/Engine/Plugins/Runtime/MeshModelingToolset/Source/ModelingComponents/Private/Drawing/PointSetComponent.cpp
@@ -107,7 +107,7 @@ public:
 				// have to convert it back to linear. The ToFColor(false) call just scales back into 0-255 space.
 				FColor color = FLinearColor::FromSRGBColor(Point.Color).ToFColor(false);
 
-				const FVector2D UV(Point.Size, Point.DepthBias);
+				const FVector2f UV(Point.Size, Point.DepthBias);
 				for (int j = 0; j < 4; ++j)
 				{
 					VertexBuffers.PositionVertexBuffer.VertexPosition(UseVertexBufferIndex + j) = Point.Position;
diff --git a/Engine/Plugins/Runtime/MeshModelingToolset/Source/ModelingComponents/Private/Drawing/TriangleSetComponent.cpp b/Engine/Plugins/Runtime/MeshModelingToolset/Source/ModelingComponents/Private/Drawing/TriangleSetComponent.cpp
--- a/Engine/Plugins/Runtime/MeshModelingToolset/Source/ModelingComponents/Private/Drawing/TriangleSetComponent.cpp
+++ b/Engine/Plugins/Runtime/MeshModelingToolset/Source/ModelingComponents/Private/Drawing/TriangleSetComponent.cpp
@@ -69,9 +69,9 @@ public:
 				VertexBuffers.StaticMeshVertexBuffer.SetVertexTangents(VertexBufferIndex + 1, FVector(1, 0, 0), FVector(0, 1, 0), Triangle.Vertex1.Normal);
 				VertexBuffers.StaticMeshVertexBuffer.SetVertexTangents(VertexBufferIndex + 2, FVector(1, 0, 0), FVector(0, 1, 0), Triangle.Vertex2.Normal);
 
-				VertexBuffers.StaticMeshVertexBuffer.SetVertexUV(VertexBufferIndex + 0, 0, Triangle.Vertex0.UV);
-				VertexBuffers.StaticMeshVertexBuffer.SetVertexUV(VertexBufferIndex + 1, 0, Triangle.Vertex1.UV);
-				VertexBuffers.StaticMeshVertexBuffer.SetVertexUV(VertexBufferIndex + 2, 0, Triangle.Vertex2.UV);
+				VertexBuffers.StaticMeshVertexBuffer.SetVertexUV(VertexBufferIndex + 0, 0, FVector2f(Triangle.Vertex0.UV));
+				VertexBuffers.StaticMeshVertexBuffer.SetVertexUV(VertexBufferIndex + 1, 0, FVector2f(Triangle.Vertex1.UV));
+				VertexBuffers.StaticMeshVertexBuffer.SetVertexUV(VertexBufferIndex + 2, 0, FVector2f(Triangle.Vertex2.UV));
 
 				VertexBuffers.ColorVertexBuffer.VertexColor(VertexBufferIndex + 0) = Triangle.Vertex0.Color;
 				VertexBuffers.ColorVertexBuffer.VertexColor(VertexBufferIndex + 1) = Triangle.Vertex1.Color;
diff --git a/Engine/Plugins/Runtime/Oculus/OculusVR/Source/OculusHMD/Private/OculusHMD_StressTester.cpp b/Engine/Plugins/Runtime/Oculus/OculusVR/Source/OculusHMD/Private/OculusHMD_StressTester.cpp
--- a/Engine/Plugins/Runtime/Oculus/OculusVR/Source/OculusHMD/Private/OculusHMD_StressTester.cpp
+++ b/Engine/Plugins/Runtime/Oculus/OculusVR/Source/OculusHMD/Private/OculusHMD_StressTester.cpp
@@ -64,10 +64,10 @@ inline FBufferRHIRef CreateTempOcculusVertexBuffer()
 	Vertices[1].Position = FVector4f(1.0f, 1.0f, 0, 1.0f);
 	Vertices[2].Position = FVector4f(-1.0f, -1.0f, 0, 1.0f);
 	Vertices[3].Position = FVector4f(1.0f, -1.0f, 0, 1.0f);
-	Vertices[0].UV = FVector2D(0, 0);
-	Vertices[1].UV = FVector2D(1, 0);
-	Vertices[2].UV = FVector2D(0, 1);
-	Vertices[3].UV = FVector2D(1, 1);
+	Vertices[0].UV = FVector2f(0, 0);
+	Vertices[1].UV = FVector2f(1, 0);
+	Vertices[2].UV = FVector2f(0, 1);
+	Vertices[3].UV = FVector2f(1, 1);
 	RHIUnlockBuffer(VertexBufferRHI);
 
 	return VertexBufferRHI;
diff --git a/Engine/Plugins/Runtime/Oculus/OculusVR/Source/OculusInput/Private/OculusHandTracking.cpp b/Engine/Plugins/Runtime/Oculus/OculusVR/Source/OculusInput/Private/OculusHandTracking.cpp
--- a/Engine/Plugins/Runtime/Oculus/OculusVR/Source/OculusInput/Private/OculusHandTracking.cpp
+++ b/Engine/Plugins/Runtime/Oculus/OculusVR/Source/OculusInput/Private/OculusHandTracking.cpp
@@ -326,7 +326,7 @@ void FOculusHandTracking::InitializeHandMesh(USkeletalMesh* SkeletalMesh, const
 		SoftVertex.TangentZ = FVector3f(Normal.x, Normal.z, Normal.y);
 		SoftVertex.TangentX = FVector3f(1.0f, 0.0f, 0.0f);
 		SoftVertex.TangentY = FVector3f(0.0f, 1.0f, 0.0f);// SoftVertex.TangentZ^ SoftVertex.TangentX* SoftVertex.TangentZ.W;
-		SoftVertex.UVs[0] = FVector2D(OvrMesh->VertexUV0[VertexIndex].x, OvrMesh->VertexUV0[VertexIndex].y);
+		SoftVertex.UVs[0] = FVector2f(OvrMesh->VertexUV0[VertexIndex].x, OvrMesh->VertexUV0[VertexIndex].y);
 
 		// Update the Bounds
 		float VertexDistSq = SoftVertex.Position.SizeSquared();
diff --git a/Engine/Plugins/Runtime/Oculus/OculusVR/Source/OculusMR/Private/OculusMR_PlaneMeshComponent.cpp b/Engine/Plugins/Runtime/Oculus/OculusVR/Source/OculusMR/Private/OculusMR_PlaneMeshComponent.cpp
--- a/Engine/Plugins/Runtime/Oculus/OculusVR/Source/OculusMR/Private/OculusMR_PlaneMeshComponent.cpp
+++ b/Engine/Plugins/Runtime/Oculus/OculusVR/Source/OculusMR/Private/OculusMR_PlaneMeshComponent.cpp
@@ -49,17 +49,17 @@ public:
 			Vert.SetTangents(TangentX, TangentY, TangentZ);
 
 			Vert.Position = Tri.Vertex0;
-			Vert.TextureCoordinate[0] = Tri.UV0;
+			Vert.TextureCoordinate[0] = FVector2f(Tri.UV0);	// LWC_TODO: Precision loss
 			Vertices[TriIdx * 3 + 0] = Vert;
 			Indices[TriIdx * 3 + 0] = TriIdx * 3 + 0;
 
 			Vert.Position = Tri.Vertex1;
-			Vert.TextureCoordinate[0] = Tri.UV1;
+			Vert.TextureCoordinate[0] = FVector2f(Tri.UV1);	// LWC_TODO: Precision loss
 			Vertices[TriIdx * 3 + 1] = Vert;
 			Indices[TriIdx * 3 + 1] = TriIdx * 3 + 1;
 
 			Vert.Position = Tri.Vertex2;
-			Vert.TextureCoordinate[0] = Tri.UV2;
+			Vert.TextureCoordinate[0] = FVector2f(Tri.UV2);	// LWC_TODO: Precision loss
 			Vertices[TriIdx * 3 + 2] = Vert;
 			Indices[TriIdx * 3 + 2] = TriIdx * 3 + 2;
 		}
diff --git a/Engine/Plugins/Runtime/ProceduralMeshComponent/Source/ProceduralMeshComponent/Private/KismetProceduralMeshLibrary.cpp b/Engine/Plugins/Runtime/ProceduralMeshComponent/Source/ProceduralMeshComponent/Private/KismetProceduralMeshLibrary.cpp
--- a/Engine/Plugins/Runtime/ProceduralMeshComponent/Source/ProceduralMeshComponent/Private/KismetProceduralMeshLibrary.cpp
+++ b/Engine/Plugins/Runtime/ProceduralMeshComponent/Source/ProceduralMeshComponent/Private/KismetProceduralMeshLibrary.cpp
@@ -433,7 +433,7 @@ static int32 GetNewIndexForOldVertIndex(int32 MeshVertIndex, TMap<int32, int32>&
 		check(Normals.Num() == Vertices.Num());
 
 		// Copy UVs
-		UVs.Add(VertexBuffers.StaticMeshVertexBuffer.GetVertexUV(MeshVertIndex, 0));
+		UVs.Add(FVector2D(VertexBuffers.StaticMeshVertexBuffer.GetVertexUV(MeshVertIndex, 0)));
 		check(UVs.Num() == Vertices.Num());
 
 		// Copy tangents
diff --git a/Engine/Plugins/Runtime/ProceduralMeshComponent/Source/ProceduralMeshComponent/Private/ProceduralMeshComponent.cpp b/Engine/Plugins/Runtime/ProceduralMeshComponent/Source/ProceduralMeshComponent/Private/ProceduralMeshComponent.cpp
--- a/Engine/Plugins/Runtime/ProceduralMeshComponent/Source/ProceduralMeshComponent/Private/ProceduralMeshComponent.cpp
+++ b/Engine/Plugins/Runtime/ProceduralMeshComponent/Source/ProceduralMeshComponent/Private/ProceduralMeshComponent.cpp
@@ -100,10 +100,10 @@ static void ConvertProcMeshToDynMeshVertex(FDynamicMeshVertex& Vert, const FProc
 {
 	Vert.Position = ProcVert.Position;
 	Vert.Color = ProcVert.Color;
-	Vert.TextureCoordinate[0] = ProcVert.UV0;
-	Vert.TextureCoordinate[1] = ProcVert.UV1;
-	Vert.TextureCoordinate[2] = ProcVert.UV2;
-	Vert.TextureCoordinate[3] = ProcVert.UV3;
+	Vert.TextureCoordinate[0] = FVector2f(ProcVert.UV0);	// LWC_TODO: Precision loss
+	Vert.TextureCoordinate[1] = FVector2f(ProcVert.UV1);	// LWC_TODO: Precision loss
+	Vert.TextureCoordinate[2] = FVector2f(ProcVert.UV2);	// LWC_TODO: Precision loss
+	Vert.TextureCoordinate[3] = FVector2f(ProcVert.UV3);	// LWC_TODO: Precision loss
 	Vert.TangentX = ProcVert.Tangent.TangentX;
 	Vert.TangentZ = ProcVert.Normal;
 	Vert.TangentZ.Vector.W = ProcVert.Tangent.bFlipTangentY ? -127 : 127;
diff --git a/Engine/Plugins/Runtime/ProceduralMeshComponent/Source/ProceduralMeshComponent/Private/ProceduralMeshConversion.cpp b/Engine/Plugins/Runtime/ProceduralMeshComponent/Source/ProceduralMeshComponent/Private/ProceduralMeshConversion.cpp
--- a/Engine/Plugins/Runtime/ProceduralMeshComponent/Source/ProceduralMeshComponent/Private/ProceduralMeshConversion.cpp
+++ b/Engine/Plugins/Runtime/ProceduralMeshComponent/Source/ProceduralMeshComponent/Private/ProceduralMeshConversion.cpp
@@ -84,7 +84,7 @@ void MeshDescriptionToProcMesh( const FMeshDescription& MeshDescription, UProced
 				SectionData.Positions.Add( VertexPositions[VertexID] );
 				SectionData.Triangles.Add( NewVertexIndex );
 				SectionData.Normals.Add( Normals[VertexInstanceID] );
-				SectionData.UV0.Add( UVs[VertexInstanceID] );
+				SectionData.UV0.Add( FVector2D(UVs[VertexInstanceID]) );
 				SectionData.Tangents.Add( FProcMeshTangent( Tangents[VertexInstanceID], BinormalSigns[VertexInstanceID] < 0.f ) );
 			}
 		}
@@ -200,10 +200,10 @@ FMeshDescription BuildMeshDescription( UProceduralMeshComponent* ProcMeshComp )
 
 			Colors[VertexInstanceID] = FLinearColor(ProcVertex.Color);
 
-			UVs.Set(VertexInstanceID, 0, ProcVertex.UV0);
-			UVs.Set(VertexInstanceID, 1, ProcVertex.UV1);
-			UVs.Set(VertexInstanceID, 2, ProcVertex.UV2);
-			UVs.Set(VertexInstanceID, 3, ProcVertex.UV3);
+			UVs.Set(VertexInstanceID, 0, FVector2f(ProcVertex.UV0));	// LWC_TODO: Precision loss
+			UVs.Set(VertexInstanceID, 1, FVector2f(ProcVertex.UV1));	// LWC_TODO: Precision loss
+			UVs.Set(VertexInstanceID, 2, FVector2f(ProcVertex.UV2));	// LWC_TODO: Precision loss
+			UVs.Set(VertexInstanceID, 3, FVector2f(ProcVertex.UV3));	// LWC_TODO: Precision loss
 		}
 
 		// Create the polygons for this section
diff --git a/Engine/Plugins/Runtime/nDisplay/Source/DisplayCluster/Private/Render/Containers/DisplayClusterRender_MeshComponentProxyData.cpp b/Engine/Plugins/Runtime/nDisplay/Source/DisplayCluster/Private/Render/Containers/DisplayClusterRender_MeshComponentProxyData.cpp
--- a/Engine/Plugins/Runtime/nDisplay/Source/DisplayCluster/Private/Render/Containers/DisplayClusterRender_MeshComponentProxyData.cpp
+++ b/Engine/Plugins/Runtime/nDisplay/Source/DisplayCluster/Private/Render/Containers/DisplayClusterRender_MeshComponentProxyData.cpp
@@ -37,8 +37,8 @@ FDisplayClusterRender_MeshComponentProxyData::FDisplayClusterRender_MeshComponen
 	for (int32 VertexIdx = 0; VertexIdx < VertexData.Num(); VertexIdx++)
 	{
 		VertexData[VertexIdx].Position = FVector4(FVector(PositionBuffer.VertexPosition(VertexIdx)));
-		VertexData[VertexIdx].UV = VertexBuffer.GetVertexUV(VertexIdx, BaseUVIndex);
-		VertexData[VertexIdx].UV_Chromakey = VertexBuffer.GetVertexUV(VertexIdx, ChromakeyUVIndex);
+		VertexData[VertexIdx].UV = FVector2D(VertexBuffer.GetVertexUV(VertexIdx, BaseUVIndex));
+		VertexData[VertexIdx].UV_Chromakey = FVector2D(VertexBuffer.GetVertexUV(VertexIdx, ChromakeyUVIndex));
 	}
 
 	UpdateData(InDataFunc);
diff --git a/Engine/Plugins/Runtime/nDisplay/Source/DisplayCluster/Private/Render/Containers/DisplayClusterRender_MeshResources.h b/Engine/Plugins/Runtime/nDisplay/Source/DisplayCluster/Private/Render/Containers/DisplayClusterRender_MeshResources.h
--- a/Engine/Plugins/Runtime/nDisplay/Source/DisplayCluster/Private/Render/Containers/DisplayClusterRender_MeshResources.h
+++ b/Engine/Plugins/Runtime/nDisplay/Source/DisplayCluster/Private/Render/Containers/DisplayClusterRender_MeshResources.h
@@ -18,9 +18,9 @@ struct FDisplayClusterMeshVertexType
 	/** type conversion double -> float */
 	FORCEINLINE void SetVertexData(const FDisplayClusterMeshVertex& In)
 	{
-		Position = (FVector4f)In.Position; // LWC_TODO: precision loss
-		UV = In.UV;
-		UV_Chromakey = In.UV_Chromakey;
+		Position = FVector4f(In.Position); // LWC_TODO: precision loss
+		UV = FVector2f(In.UV);
+		UV_Chromakey = FVector2f(In.UV_Chromakey);
 	}
 };
 
diff --git a/Engine/Plugins/Runtime/nDisplay/Source/DisplayClusterShaders/Private/Shaders/DisplayClusterShadersWarpblend_ICVFX.cpp b/Engine/Plugins/Runtime/nDisplay/Source/DisplayClusterShaders/Private/Shaders/DisplayClusterShadersWarpblend_ICVFX.cpp
--- a/Engine/Plugins/Runtime/nDisplay/Source/DisplayClusterShaders/Private/Shaders/DisplayClusterShadersWarpblend_ICVFX.cpp
+++ b/Engine/Plugins/Runtime/nDisplay/Source/DisplayClusterShaders/Private/Shaders/DisplayClusterShadersWarpblend_ICVFX.cpp
@@ -611,7 +611,7 @@ public:
 
 			RenderPassData.PSParameters.ChromakeyMarkerScale    = Camera.ChromakeyMarkersScale;
 			RenderPassData.PSParameters.ChromakeyMarkerDistance = Camera.ChromakeyMarkersDistance;
-			RenderPassData.PSParameters.ChromakeyMarkerOffset   = Camera.ChromakeyMarkersOffset;
+			RenderPassData.PSParameters.ChromakeyMarkerOffset   = FVector2f(Camera.ChromakeyMarkersOffset);
 
 			RenderPassData.PSPermutationVector.Set<IcvfxShaderPermutation::FIcvfxShaderChromakeyMarker>(true);
 			return true;
diff --git a/Engine/Plugins/VirtualProduction/CameraCalibration/Source/CameraCalibrationEditor/Private/Calibrators/CameraNodalOffsetAlgoOpticalAxis.cpp b/Engine/Plugins/VirtualProduction/CameraCalibration/Source/CameraCalibrationEditor/Private/Calibrators/CameraNodalOffsetAlgoOpticalAxis.cpp
--- a/Engine/Plugins/VirtualProduction/CameraCalibration/Source/CameraCalibrationEditor/Private/Calibrators/CameraNodalOffsetAlgoOpticalAxis.cpp
+++ b/Engine/Plugins/VirtualProduction/CameraCalibration/Source/CameraCalibrationEditor/Private/Calibrators/CameraNodalOffsetAlgoOpticalAxis.cpp
@@ -123,7 +123,7 @@ void UCameraNodalOffsetAlgoOpticalAxis::Tick(float DeltaTime)
 	if (LastPrincipalPoint != DistortionState.ImageCenter.PrincipalPoint)
 	{
 		FCrosshairOverlayParams OverlayParams;
-		OverlayParams.PrincipalPoint = DistortionState.ImageCenter.PrincipalPoint;
+		OverlayParams.PrincipalPoint = FVector2f(DistortionState.ImageCenter.PrincipalPoint);
 		OverlayRendering::DrawCrosshairOverlay(OverlayTexture, OverlayParams);
 		StepsController->RefreshOverlay();
 
diff --git a/Engine/Plugins/VirtualProduction/CameraCalibrationCore/Source/CameraCalibrationCore/Private/CalibratedMapProcessor.cpp b/Engine/Plugins/VirtualProduction/CameraCalibrationCore/Source/CameraCalibrationCore/Private/CalibratedMapProcessor.cpp
--- a/Engine/Plugins/VirtualProduction/CameraCalibrationCore/Source/CameraCalibrationCore/Private/CalibratedMapProcessor.cpp
+++ b/Engine/Plugins/VirtualProduction/CameraCalibrationCore/Source/CameraCalibrationCore/Private/CalibratedMapProcessor.cpp
@@ -187,7 +187,7 @@ void FCalibratedMapProcessor::ExecuteJob(TSharedPtr<FDerivedDistortionDataJob> J
 			FCalibratedMapDerivedDataCS::FParameters* Parameters = GraphBuilder.AllocParameters<FCalibratedMapDerivedDataCS::FParameters>();
             Parameters->DistortionSTMap = SourceSTMap;
            	Parameters->DistortionSTMapSampler = TStaticSamplerState<SF_Bilinear, AM_Clamp, AM_Clamp, AM_Clamp>::GetRHI();
-           	Parameters->TexelSize = FVector2D(1.0f / UndistortionDisplacementMap->Desc.Extent.X, 1.0f / UndistortionDisplacementMap->Desc.Extent.Y);
+           	Parameters->TexelSize = FVector2f(1.0f / UndistortionDisplacementMap->Desc.Extent.X, 1.0f / UndistortionDisplacementMap->Desc.Extent.Y);
             Parameters->TextureSize = FIntPoint(UndistortionDisplacementMap->Desc.Extent.X, UndistortionDisplacementMap->Desc.Extent.Y);
 
             //Create UAVs for compute shader outputs
diff --git a/Engine/Plugins/VirtualProduction/CameraCalibrationCore/Source/CameraCalibrationCore/Private/DistortionRenderingUtils.cpp b/Engine/Plugins/VirtualProduction/CameraCalibrationCore/Source/CameraCalibrationCore/Private/DistortionRenderingUtils.cpp
--- a/Engine/Plugins/VirtualProduction/CameraCalibrationCore/Source/CameraCalibrationCore/Private/DistortionRenderingUtils.cpp
+++ b/Engine/Plugins/VirtualProduction/CameraCalibrationCore/Source/CameraCalibrationCore/Private/DistortionRenderingUtils.cpp
@@ -57,7 +57,7 @@ namespace DistortionRenderingUtils
 				InputPointsResourceArray.Reserve(NumPoints);
 				for (int32 Index = 0; Index < NumPoints; Index++)
 				{
-					InputPointsResourceArray.Add(Points[Index]);
+					InputPointsResourceArray.Add(FVector2f(Points[Index]));	// LWC_TODO: Precision loss
 				}
 
 				TResourceArray<FVector2f> EmptyBuffer;
@@ -95,7 +95,7 @@ namespace DistortionRenderingUtils
 				{
 					for (int32 Index = 0; Index < NumPoints; ++Index)
 					{
-						OutUndistortedPoints[Index] = UndistortedPointData[Index];
+						OutUndistortedPoints[Index] = FVector2D(UndistortedPointData[Index]);
 					}
 				}
 
diff --git a/Engine/Plugins/VirtualProduction/CameraCalibrationCore/Source/CameraCalibrationCore/Private/LensFileRendering.cpp b/Engine/Plugins/VirtualProduction/CameraCalibrationCore/Source/CameraCalibrationCore/Private/LensFileRendering.cpp
--- a/Engine/Plugins/VirtualProduction/CameraCalibrationCore/Source/CameraCalibrationCore/Private/LensFileRendering.cpp
+++ b/Engine/Plugins/VirtualProduction/CameraCalibrationCore/Source/CameraCalibrationCore/Private/LensFileRendering.cpp
@@ -139,8 +139,8 @@ namespace LensFileRendering
 			PassParameters->OutputTextureExtent = OutputTexture->Desc.Extent;
 			PassParameters->SourceTextureSampler = TStaticSamplerState<SF_Bilinear, AM_Clamp, AM_Clamp, AM_Clamp>::GetRHI(); 
 			PassParameters->SourceTextureOne = TextureOne;
-			PassParameters->FxFyScale = BlendParams.FxFyScale;
-			PassParameters->PrincipalPoint = BlendParams.PrincipalPoint;
+			PassParameters->FxFyScale = FVector2f(BlendParams.FxFyScale);	// LWC_TODO: Precision loss
+			PassParameters->PrincipalPoint = FVector2f(BlendParams.PrincipalPoint);	// LWC_TODO: Precision loss
 
 			//Setup parameters based on blending type
 			FDisplacementMapBlendPS::FPermutationDomain PermutationVector;
diff --git a/Engine/Plugins/VirtualProduction/DMX/DMXPixelMapping/Source/DMXPixelMappingEditor/Private/Views/SDMXPixelMappingDesignerView.cpp b/Engine/Plugins/VirtualProduction/DMX/DMXPixelMapping/Source/DMXPixelMappingEditor/Private/Views/SDMXPixelMappingDesignerView.cpp
--- a/Engine/Plugins/VirtualProduction/DMX/DMXPixelMapping/Source/DMXPixelMappingEditor/Private/Views/SDMXPixelMappingDesignerView.cpp
+++ b/Engine/Plugins/VirtualProduction/DMX/DMXPixelMapping/Source/DMXPixelMappingEditor/Private/Views/SDMXPixelMappingDesignerView.cpp
@@ -790,7 +790,7 @@ FVector2D SDMXPixelMappingDesignerView::GetExtensionPosition(TSharedPtr<SDMXPixe
 				return FVector2D::ZeroVector;
 			}();
 
-			const FVector2D SelectedWidgetScale = SelectedComponentGeometry.GetAccumulatedRenderTransform().GetMatrix().GetScale().GetVector();
+			const FVector2D SelectedWidgetScale = FVector2D(SelectedComponentGeometry.GetAccumulatedRenderTransform().GetMatrix().GetScale().GetVector());
 
 			const FVector2D ApplicationScaledOffset = Handle->GetOffset() * GetDesignerGeometry().Scale;
 
diff --git a/Engine/Plugins/VirtualProduction/DMX/DMXPixelMapping/Source/DMXPixelMappingRenderer/Private/DMXPixelMappingRenderer.cpp b/Engine/Plugins/VirtualProduction/DMX/DMXPixelMapping/Source/DMXPixelMappingRenderer/Private/DMXPixelMappingRenderer.cpp
--- a/Engine/Plugins/VirtualProduction/DMX/DMXPixelMapping/Source/DMXPixelMappingRenderer/Private/DMXPixelMappingRenderer.cpp
+++ b/Engine/Plugins/VirtualProduction/DMX/DMXPixelMapping/Source/DMXPixelMappingRenderer/Private/DMXPixelMappingRenderer.cpp
@@ -198,9 +198,9 @@ void FDMXPixelMappingRenderer::DownsampleRender(
 					PSParameters.InputSampler = TStaticSamplerState<SF_Trilinear, AM_Clamp, AM_Clamp, AM_Clamp>::GetRHI();
 					PSParameters.InputTextureSize = InputTextureSize;
 					PSParameters.OutputTextureSize = OutputTextureSize;
-					PSParameters.PixelFactor = (FVector4f)PixelParam.PixelFactor;
+					PSParameters.PixelFactor = FVector4f(PixelParam.PixelFactor);
 					PSParameters.InvertPixel = PixelParam.InvertPixel;
-					PSParameters.UVCellSize = PixelParam.UVCellSize;
+					PSParameters.UVCellSize = FVector2f(PixelParam.UVCellSize);
 					SetShaderParameters(RHICmdList, PixelShader, PixelShader.GetPixelShader(), PSParameters);
 
 					// Draw a two triangle on the entire viewport.
