Add 'Optional' EditorOnly data for UMaterialInterface and UMaterialFunctionInterface
These are separate UObject hierarchies that store editor-only UPROPERTIES, but can be included with cooked content, which allows full editor support.
In principle, all editor-only properties could be moved over.  So far, this has been limited to UMaterialExpressions, and data related to material parameters.
FStaticParameterSet, FMaterialLayersParameters, and FMaterialCachedExpressionData have all been split into separate editor-only/non-editor-only classes,
which allows the editor-only portion to be stored on the optional editor-only UObject.

diff --git a/Engine/Plugins/Editor/AssetSearch/Source/Private/Indexers/MaterialExpressionIndexer.cpp b/Engine/Plugins/Editor/AssetSearch/Source/Private/Indexers/MaterialExpressionIndexer.cpp
--- a/Engine/Plugins/Editor/AssetSearch/Source/Private/Indexers/MaterialExpressionIndexer.cpp
+++ b/Engine/Plugins/Editor/AssetSearch/Source/Private/Indexers/MaterialExpressionIndexer.cpp
@@ -50,9 +50,9 @@ void FMaterialExpressionIndexer::IndexParameters(const UObject* InAssetObject, F
 	const UMaterialFunction* MaterialFunction = Cast<UMaterialFunction>(InAssetObject);
 	const UMaterial* Material = Cast<UMaterial>(InAssetObject);
 
-	TArray<UMaterialExpression*> Expressions;
-	Expressions = MaterialFunction ? MaterialFunction->FunctionExpressions : Expressions;
-	Expressions = Material ? Material->Expressions : Expressions;
+	TConstArrayView<TObjectPtr<UMaterialExpression>> Expressions;
+	Expressions = MaterialFunction ? MaterialFunction->GetExpressions() : Expressions;
+	Expressions = Material ? Material->GetExpressions() : Expressions;
 
 	for (const UMaterialExpression* Expression : Expressions)
 	{
diff --git a/Engine/Plugins/Editor/GLTFImporter/Source/GLTFImporter/Private/GLTFMaterialElement.cpp b/Engine/Plugins/Editor/GLTFImporter/Source/GLTFImporter/Private/GLTFMaterialElement.cpp
--- a/Engine/Plugins/Editor/GLTFImporter/Source/GLTFImporter/Private/GLTFMaterialElement.cpp
+++ b/Engine/Plugins/Editor/GLTFImporter/Source/GLTFImporter/Private/GLTFMaterialElement.cpp
@@ -30,11 +30,11 @@ namespace GLTFImporterImpl
 
 		if (Parent->IsA<UMaterial>())
 		{
-			Cast<UMaterial>(Parent)->Expressions.Push(Expression);
+			Cast<UMaterial>(Parent)->GetExpressionCollection().AddExpression(Expression);
 		}
 		else if (Parent->IsA<UMaterialFunction>())
 		{
-			Cast<UMaterialFunction>(Parent)->FunctionExpressions.Push(Expression);
+			Cast<UMaterialFunction>(Parent)->GetExpressionCollection().AddExpression(Expression);
 		}
 
 		return Expression;
@@ -267,17 +267,19 @@ void FGLTFMaterialElement::Finalize()
 
 	TArray<TStrongObjectPtr<UMaterialExpression> > MaterialExpressions;
 	CreateExpressions(MaterialExpressions);
-	ConnectInput(BaseColor, MaterialExpressions, Material->BaseColor);
-	ConnectInput(Metallic, MaterialExpressions, Material->Metallic);
-	ConnectInput(Specular, MaterialExpressions, Material->Specular);
-	ConnectInput(Roughness, MaterialExpressions, Material->Roughness);
-	ConnectInput(EmissiveColor, MaterialExpressions, Material->EmissiveColor);
-	ConnectInput(Opacity, MaterialExpressions, Material->Opacity);
-	ConnectInput(Refraction, MaterialExpressions, Material->Refraction);
-	ConnectInput(Normal, MaterialExpressions, Material->Normal);
-	ConnectInput(AmbientOcclusion, MaterialExpressions, Material->AmbientOcclusion);
-	ConnectInput(ClearCoat, MaterialExpressions, Material->ClearCoat);
-	ConnectInput(ClearCoatRoughness, MaterialExpressions, Material->ClearCoatRoughness);
+
+	UMaterialEditorOnlyData* MaterialEditorOnly = Material->GetEditorOnlyData();
+	ConnectInput(BaseColor, MaterialExpressions, MaterialEditorOnly->BaseColor);
+	ConnectInput(Metallic, MaterialExpressions, MaterialEditorOnly->Metallic);
+	ConnectInput(Specular, MaterialExpressions, MaterialEditorOnly->Specular);
+	ConnectInput(Roughness, MaterialExpressions, MaterialEditorOnly->Roughness);
+	ConnectInput(EmissiveColor, MaterialExpressions, MaterialEditorOnly->EmissiveColor);
+	ConnectInput(Opacity, MaterialExpressions, MaterialEditorOnly->Opacity);
+	ConnectInput(Refraction, MaterialExpressions, MaterialEditorOnly->Refraction);
+	ConnectInput(Normal, MaterialExpressions, MaterialEditorOnly->Normal);
+	ConnectInput(AmbientOcclusion, MaterialExpressions, MaterialEditorOnly->AmbientOcclusion);
+	ConnectInput(ClearCoat, MaterialExpressions, MaterialEditorOnly->ClearCoat);
+	ConnectInput(ClearCoatRoughness, MaterialExpressions, MaterialEditorOnly->ClearCoatRoughness);
 
 	// Handle transmission materials (they add a special output node to the graph)
 	if (ThinTranslucentMaterialOutput)
diff --git a/Engine/Plugins/Editor/SpeedTreeImporter/Source/SpeedTreeImporter/Private/SpeedTreeImportFactory.cpp b/Engine/Plugins/Editor/SpeedTreeImporter/Source/SpeedTreeImporter/Private/SpeedTreeImportFactory.cpp
--- a/Engine/Plugins/Editor/SpeedTreeImporter/Source/SpeedTreeImporter/Private/SpeedTreeImportFactory.cpp
+++ b/Engine/Plugins/Editor/SpeedTreeImporter/Source/SpeedTreeImporter/Private/SpeedTreeImportFactory.cpp
@@ -567,9 +567,8 @@ void LayoutMaterial(UMaterialInterface* MaterialInterface, bool bOffsetOddColumn
 	{
 		bContinue = false;
 
-		for (int32 ExpressionIndex = 0; ExpressionIndex < Material->Expressions.Num(); ++ExpressionIndex)
+		for (UMaterialExpression* Expression : Material->GetExpressions())
 		{
-			UMaterialExpression* Expression = Material->Expressions[ExpressionIndex];
 			Expression->MaterialExpressionEditorX = FMath::Min(Expression->MaterialExpressionEditorX, -Width);
 
 			TArray<FExpressionInput*> Inputs = Expression->GetInputs();
@@ -594,10 +593,8 @@ void LayoutMaterial(UMaterialInterface* MaterialInterface, bool bOffsetOddColumn
 	while (bContinue)
 	{
 		TArray<UMaterialExpression*> ColumnExpressions;
-		for (int32 ExpressionIndex = 0; ExpressionIndex < Material->Expressions.Num(); ++ExpressionIndex)
+		for (UMaterialExpression* Expression : Material->GetExpressions())
 		{
-			UMaterialExpression* Expression = Material->Expressions[ExpressionIndex];
-
 			if (Expression->MaterialExpressionEditorX == -Width * Column)
 			{
 				Expression->MaterialExpressionEditorY = 0;
@@ -616,9 +613,8 @@ void LayoutMaterial(UMaterialInterface* MaterialInterface, bool bOffsetOddColumn
 				}
 				
 				// all the outputs to other expressions
-				for (int32 OtherExpressionIndex = 0; OtherExpressionIndex < Material->Expressions.Num(); ++OtherExpressionIndex)
+				for (UMaterialExpression* OtherExpression : Material->GetExpressions())
 				{
-					UMaterialExpression* OtherExpression = Material->Expressions[OtherExpressionIndex];
 					TArray<FExpressionInput*> Inputs = OtherExpression->GetInputs();
 					for (int32 InputIndex = 0; InputIndex < Inputs.Num(); ++InputIndex)
 					{
@@ -725,6 +721,8 @@ UMaterialInterface* CreateSpeedTreeMaterial7(UObject* Parent, FString MaterialFu
 	// Keep track of the processed materials
 	ImportContext.ImportedMaterials.Add(MaterialFullName, UnrealMaterialInterface);
 
+	UMaterialEditorOnlyData* UnrealMaterialEditorOnly = UnrealMaterial->GetEditorOnlyData();
+
 	if (!RenderState->m_bDiffuseAlphaMaskIsOpaque && !RenderState->m_bBranchesPresent && !RenderState->m_bRigidMeshesPresent)
 	{
 		UnrealMaterial->BlendMode = BLEND_Masked;
@@ -737,7 +735,7 @@ UMaterialInterface* CreateSpeedTreeMaterial7(UObject* Parent, FString MaterialFu
 	{
 		UMaterialExpressionTextureCoordinate* SeamTexcoordExpression = NewObject<UMaterialExpressionTextureCoordinate>(UnrealMaterial);
 		SeamTexcoordExpression->CoordinateIndex = 4;
-		UnrealMaterial->Expressions.Add(SeamTexcoordExpression);
+		UnrealMaterial->GetExpressionCollection().AddExpression(SeamTexcoordExpression);
 
 		UMaterialExpressionComponentMask* ComponentMaskExpression = NewObject<UMaterialExpressionComponentMask>(UnrealMaterial);
 		ComponentMaskExpression->R = 0;
@@ -745,16 +743,16 @@ UMaterialInterface* CreateSpeedTreeMaterial7(UObject* Parent, FString MaterialFu
 		ComponentMaskExpression->B = 0;
 		ComponentMaskExpression->A = 0;
 		ComponentMaskExpression->Input.Expression = SeamTexcoordExpression;
-		UnrealMaterial->Expressions.Add(ComponentMaskExpression);
+		UnrealMaterial->GetExpressionCollection().AddExpression(ComponentMaskExpression);
 
 		UMaterialExpressionPower* PowerExpression = NewObject<UMaterialExpressionPower>(UnrealMaterial);
 		PowerExpression->Base.Expression = ComponentMaskExpression;
 		PowerExpression->ConstExponent = RenderState->m_fBranchSeamWeight;
-		UnrealMaterial->Expressions.Add(PowerExpression);
+		UnrealMaterial->GetExpressionCollection().AddExpression(PowerExpression);
 
 		BranchSeamAmount = NewObject<UMaterialExpressionClamp>(UnrealMaterial);
 		BranchSeamAmount->Input.Expression = PowerExpression;
-		UnrealMaterial->Expressions.Add(BranchSeamAmount);
+		UnrealMaterial->GetExpressionCollection().AddExpression(BranchSeamAmount);
 	}
 
 	// textures and properties
@@ -765,37 +763,37 @@ UMaterialInterface* CreateSpeedTreeMaterial7(UObject* Parent, FString MaterialFu
 		UMaterialExpressionTextureSample* TextureExpression = NewObject<UMaterialExpressionTextureSample>(UnrealMaterial);
 		TextureExpression->Texture = DiffuseTexture;
 		TextureExpression->SamplerType = SAMPLERTYPE_Color;
-		UnrealMaterial->Expressions.Add(TextureExpression);
+		UnrealMaterial->GetExpressionCollection().AddExpression(TextureExpression);
 
 		// hook to the material diffuse/mask
-		UnrealMaterial->BaseColor.Expression = TextureExpression;
-		UnrealMaterial->OpacityMask.Expression = TextureExpression;
-		UnrealMaterial->OpacityMask.Mask = UnrealMaterial->OpacityMask.Expression->GetOutputs()[0].Mask;
-		UnrealMaterial->OpacityMask.MaskR = 0;
-		UnrealMaterial->OpacityMask.MaskG = 0;
-		UnrealMaterial->OpacityMask.MaskB = 0;
-		UnrealMaterial->OpacityMask.MaskA = 1;
+		UnrealMaterialEditorOnly->BaseColor.Expression = TextureExpression;
+		UnrealMaterialEditorOnly->OpacityMask.Expression = TextureExpression;
+		UnrealMaterialEditorOnly->OpacityMask.Mask = UnrealMaterialEditorOnly->OpacityMask.Expression->GetOutputs()[0].Mask;
+		UnrealMaterialEditorOnly->OpacityMask.MaskR = 0;
+		UnrealMaterialEditorOnly->OpacityMask.MaskG = 0;
+		UnrealMaterialEditorOnly->OpacityMask.MaskB = 0;
+		UnrealMaterialEditorOnly->OpacityMask.MaskA = 1;
 
 		if (BranchSeamAmount != NULL)
 		{
 			// perform branch seam smoothing
 			UMaterialExpressionTextureCoordinate* SeamTexcoordExpression = NewObject<UMaterialExpressionTextureCoordinate>(UnrealMaterial);
 			SeamTexcoordExpression->CoordinateIndex = 6;
-			UnrealMaterial->Expressions.Add(SeamTexcoordExpression);
+			UnrealMaterial->GetExpressionCollection().AddExpression(SeamTexcoordExpression);
 
 			UMaterialExpressionTextureSample* SeamTextureExpression = NewObject<UMaterialExpressionTextureSample>(UnrealMaterial);
 			SeamTextureExpression->Texture = DiffuseTexture;
 			SeamTextureExpression->SamplerType = SAMPLERTYPE_Color;
 			SeamTextureExpression->Coordinates.Expression = SeamTexcoordExpression;
-			UnrealMaterial->Expressions.Add(SeamTextureExpression);
+			UnrealMaterial->GetExpressionCollection().AddExpression(SeamTextureExpression);
 
 			UMaterialExpressionLinearInterpolate* InterpolateExpression = NewObject<UMaterialExpressionLinearInterpolate>(UnrealMaterial);
 			InterpolateExpression->A.Expression = SeamTextureExpression;
 			InterpolateExpression->B.Expression = TextureExpression;
 			InterpolateExpression->Alpha.Expression = BranchSeamAmount;
-			UnrealMaterial->Expressions.Add(InterpolateExpression);
+			UnrealMaterial->GetExpressionCollection().AddExpression(InterpolateExpression);
 
-			UnrealMaterial->BaseColor.Expression = InterpolateExpression;
+			UnrealMaterialEditorOnly->BaseColor.Expression = InterpolateExpression;
 		}
 
 		if (RenderState->m_bBranchesPresent && SpeedTreeImportData->IncludeDetailMapCheck)
@@ -806,18 +804,18 @@ UMaterialInterface* CreateSpeedTreeMaterial7(UObject* Parent, FString MaterialFu
 				// add/find UVSet
 				UMaterialExpressionTextureCoordinate* DetailTexcoordExpression = NewObject<UMaterialExpressionTextureCoordinate>(UnrealMaterial);
 				DetailTexcoordExpression->CoordinateIndex = 5;
-				UnrealMaterial->Expressions.Add(DetailTexcoordExpression);
+				UnrealMaterial->GetExpressionCollection().AddExpression(DetailTexcoordExpression);
 				
 				// make texture sampler
 				UMaterialExpressionTextureSample* DetailTextureExpression = NewObject<UMaterialExpressionTextureSample>(UnrealMaterial);
 				DetailTextureExpression->Texture = DetailTexture;
 				DetailTextureExpression->SamplerType = SAMPLERTYPE_Color;
 				DetailTextureExpression->Coordinates.Expression = DetailTexcoordExpression;
-				UnrealMaterial->Expressions.Add(DetailTextureExpression);
+				UnrealMaterial->GetExpressionCollection().AddExpression(DetailTextureExpression);
 
 				// interpolate the detail
 				UMaterialExpressionLinearInterpolate* InterpolateExpression = NewObject<UMaterialExpressionLinearInterpolate>(UnrealMaterial);
-				InterpolateExpression->A.Expression = UnrealMaterial->BaseColor.Expression;
+				InterpolateExpression->A.Expression = UnrealMaterialEditorOnly->BaseColor.Expression;
 				InterpolateExpression->B.Expression = DetailTextureExpression;
 				InterpolateExpression->Alpha.Expression = DetailTextureExpression;
 				InterpolateExpression->Alpha.Mask = DetailTextureExpression->GetOutputs()[0].Mask;
@@ -825,10 +823,10 @@ UMaterialInterface* CreateSpeedTreeMaterial7(UObject* Parent, FString MaterialFu
 				InterpolateExpression->Alpha.MaskG = 0;
 				InterpolateExpression->Alpha.MaskB = 0;
 				InterpolateExpression->Alpha.MaskA = 1;
-				UnrealMaterial->Expressions.Add(InterpolateExpression);
+				UnrealMaterial->GetExpressionCollection().AddExpression(InterpolateExpression);
 
 				// hook final to diffuse
-				UnrealMaterial->BaseColor.Expression = InterpolateExpression;
+				UnrealMaterialEditorOnly->BaseColor.Expression = InterpolateExpression;
 			}
 		}
 	}
@@ -844,8 +842,8 @@ UMaterialInterface* CreateSpeedTreeMaterial7(UObject* Parent, FString MaterialFu
 			TextureExpression->Texture = SpecularTexture;
 			TextureExpression->SamplerType = SAMPLERTYPE_Color;
 			
-			UnrealMaterial->Expressions.Add(TextureExpression);
-			UnrealMaterial->Specular.Expression = TextureExpression;
+			UnrealMaterial->GetExpressionCollection().AddExpression(TextureExpression);
+			UnrealMaterialEditorOnly->Specular.Expression = TextureExpression;
 			bMadeSpecular = true;
 		}
 	}
@@ -854,8 +852,8 @@ UMaterialInterface* CreateSpeedTreeMaterial7(UObject* Parent, FString MaterialFu
 	{
 		UMaterialExpressionConstant* ZeroExpression = NewObject<UMaterialExpressionConstant>(UnrealMaterial);
 		ZeroExpression->R = 0.0f;
-		UnrealMaterial->Expressions.Add(ZeroExpression);
-		UnrealMaterial->Specular.Expression = ZeroExpression;
+		UnrealMaterial->GetExpressionCollection().AddExpression(ZeroExpression);
+		UnrealMaterialEditorOnly->Specular.Expression = ZeroExpression;
 	}
 
 	if (SpeedTreeImportData->IncludeNormalMapCheck)
@@ -868,29 +866,29 @@ UMaterialInterface* CreateSpeedTreeMaterial7(UObject* Parent, FString MaterialFu
 			TextureExpression->Texture = NormalTexture;
 			TextureExpression->SamplerType = SAMPLERTYPE_Normal;
 			
-			UnrealMaterial->Expressions.Add(TextureExpression);
-			UnrealMaterial->Normal.Expression = TextureExpression;
+			UnrealMaterial->GetExpressionCollection().AddExpression(TextureExpression);
+			UnrealMaterialEditorOnly->Normal.Expression = TextureExpression;
 
 			if (BranchSeamAmount != NULL)
 			{
 				// perform branch seam smoothing
 				UMaterialExpressionTextureCoordinate* SeamTexcoordExpression = NewObject<UMaterialExpressionTextureCoordinate>(UnrealMaterial);
 				SeamTexcoordExpression->CoordinateIndex = 6;
-				UnrealMaterial->Expressions.Add(SeamTexcoordExpression);
+				UnrealMaterial->GetExpressionCollection().AddExpression(SeamTexcoordExpression);
 
 				UMaterialExpressionTextureSample* SeamTextureExpression = NewObject<UMaterialExpressionTextureSample>(UnrealMaterial);
 				SeamTextureExpression->Texture = NormalTexture;
 				SeamTextureExpression->SamplerType = SAMPLERTYPE_Normal;
 				SeamTextureExpression->Coordinates.Expression = SeamTexcoordExpression;
-				UnrealMaterial->Expressions.Add(SeamTextureExpression);
+				UnrealMaterial->GetExpressionCollection().AddExpression(SeamTextureExpression);
 
 				UMaterialExpressionLinearInterpolate* InterpolateExpression = NewObject<UMaterialExpressionLinearInterpolate>(UnrealMaterial);
 				InterpolateExpression->A.Expression = SeamTextureExpression;
 				InterpolateExpression->B.Expression = TextureExpression;
 				InterpolateExpression->Alpha.Expression = BranchSeamAmount;
-				UnrealMaterial->Expressions.Add(InterpolateExpression);
+				UnrealMaterial->GetExpressionCollection().AddExpression(InterpolateExpression);
 
-				UnrealMaterial->Normal.Expression = InterpolateExpression;
+				UnrealMaterialEditorOnly->Normal.Expression = InterpolateExpression;
 			}
 		}
 	}
@@ -916,69 +914,69 @@ UMaterialInterface* CreateSpeedTreeMaterial7(UObject* Parent, FString MaterialFu
 		else
 			SpeedTreeExpression->GeometryType = STG_FacingLeaf;
 
-		UnrealMaterial->Expressions.Add(SpeedTreeExpression);
-		UnrealMaterial->WorldPositionOffset.Expression = SpeedTreeExpression;
+		UnrealMaterial->GetExpressionCollection().AddExpression(SpeedTreeExpression);
+		UnrealMaterialEditorOnly->WorldPositionOffset.Expression = SpeedTreeExpression;
 	}
 
 	if (SpeedTreeImportData->IncludeSpeedTreeAO &&
 		!(RenderState->m_bVertBillboard || RenderState->m_bHorzBillboard))
 	{
 		UMaterialExpressionVertexColor* VertexColor = NewObject<UMaterialExpressionVertexColor>(UnrealMaterial);
-		UnrealMaterial->Expressions.Add(VertexColor);
-		UnrealMaterial->AmbientOcclusion.Expression = VertexColor;
-		UnrealMaterial->AmbientOcclusion.Mask = VertexColor->GetOutputs()[0].Mask;
-		UnrealMaterial->AmbientOcclusion.MaskR = 1;
-		UnrealMaterial->AmbientOcclusion.MaskG = 0;
-		UnrealMaterial->AmbientOcclusion.MaskB = 0;
-		UnrealMaterial->AmbientOcclusion.MaskA = 0;
+		UnrealMaterial->GetExpressionCollection().AddExpression(VertexColor);
+		UnrealMaterialEditorOnly->AmbientOcclusion.Expression = VertexColor;
+		UnrealMaterialEditorOnly->AmbientOcclusion.Mask = VertexColor->GetOutputs()[0].Mask;
+		UnrealMaterialEditorOnly->AmbientOcclusion.MaskR = 1;
+		UnrealMaterialEditorOnly->AmbientOcclusion.MaskG = 0;
+		UnrealMaterialEditorOnly->AmbientOcclusion.MaskB = 0;
+		UnrealMaterialEditorOnly->AmbientOcclusion.MaskA = 0;
 	}
 
 	// Unreal flips normals for two-sided materials. SpeedTrees don't need that
 	if (UnrealMaterial->TwoSided)
 	{
 		UMaterialExpressionTwoSidedSign* TwoSidedSign = NewObject<UMaterialExpressionTwoSidedSign>(UnrealMaterial);
-		UnrealMaterial->Expressions.Add(TwoSidedSign);
+		UnrealMaterial->GetExpressionCollection().AddExpression(TwoSidedSign);
 
 		auto Multiply = NewObject<UMaterialExpressionMultiply>(UnrealMaterial);
-		UnrealMaterial->Expressions.Add(Multiply);
+		UnrealMaterial->GetExpressionCollection().AddExpression(Multiply);
 		Multiply->A.Expression = TwoSidedSign;
 
-		if (UnrealMaterial->Normal.Expression == NULL)
+		if (UnrealMaterialEditorOnly->Normal.Expression == NULL)
 		{
 			auto VertexNormalExpression = NewObject<UMaterialExpressionConstant3Vector>(UnrealMaterial);
-			UnrealMaterial->Expressions.Add(VertexNormalExpression);
+			UnrealMaterial->GetExpressionCollection().AddExpression(VertexNormalExpression);
 			VertexNormalExpression->Constant = FLinearColor(0.0f, 0.0f, 1.0f);
 
 			Multiply->B.Expression = VertexNormalExpression;
 		}
 		else
 		{
-			Multiply->B.Expression = UnrealMaterial->Normal.Expression;
+			Multiply->B.Expression = UnrealMaterialEditorOnly->Normal.Expression;
 		}
 
-		UnrealMaterial->Normal.Expression = Multiply;
+		UnrealMaterialEditorOnly->Normal.Expression = Multiply;
 	}
 
-	if (SpeedTreeImportData->IncludeColorAdjustment && UnrealMaterial->BaseColor.Expression != NULL &&
+	if (SpeedTreeImportData->IncludeColorAdjustment && UnrealMaterialEditorOnly->BaseColor.Expression != NULL &&
 		(RenderState->m_bLeavesPresent || RenderState->m_bFacingLeavesPresent || RenderState->m_bVertBillboard || RenderState->m_bHorzBillboard))
 	{
 		UMaterialFunction* ColorVariationFunction = LoadObject<UMaterialFunction>(NULL, TEXT("/Engine/Functions/Engine_MaterialFunctions01/SpeedTree/SpeedTreeColorVariation.SpeedTreeColorVariation"), NULL, LOAD_None, NULL);
 		if (ColorVariationFunction)
 		{
 			UMaterialExpressionMaterialFunctionCall* ColorVariation = NewObject<UMaterialExpressionMaterialFunctionCall>(UnrealMaterial);
-			UnrealMaterial->Expressions.Add(ColorVariation);
+			UnrealMaterial->GetExpressionCollection().AddExpression(ColorVariation);
 
 			ColorVariation->MaterialFunction = ColorVariationFunction;
 			ColorVariation->UpdateFromFunctionResource( );
 
-			ColorVariation->GetInput(0)->Expression = UnrealMaterial->BaseColor.Expression;
-			ColorVariation->GetInput(0)->Mask = UnrealMaterial->BaseColor.Expression->GetOutputs()[0].Mask;
+			ColorVariation->GetInput(0)->Expression = UnrealMaterialEditorOnly->BaseColor.Expression;
+			ColorVariation->GetInput(0)->Mask = UnrealMaterialEditorOnly->BaseColor.Expression->GetOutputs()[0].Mask;
 			ColorVariation->GetInput(0)->MaskR = 1;
 			ColorVariation->GetInput(0)->MaskG = 1;
 			ColorVariation->GetInput(0)->MaskB = 1;
 			ColorVariation->GetInput(0)->MaskA = 0;
 
-			UnrealMaterial->BaseColor.Expression = ColorVariation;
+			UnrealMaterialEditorOnly->BaseColor.Expression = ColorVariation;
 		}
 	}
 	
@@ -1052,6 +1050,8 @@ UMaterialInterface* CreateSpeedTreeMaterial8(UObject* Parent, FString MaterialFu
 	// Keep track of the processed materials
 	ImportContext.ImportedMaterials.Add(MaterialFullName, UnrealMaterialInterface);
 
+	UMaterialEditorOnlyData* UnrealMaterialEditorOnly = UnrealMaterial->GetEditorOnlyData();
+
 	// basic setup
 	UnrealMaterial->TwoSided = SpeedTreeMaterial.TwoSided();
 	UnrealMaterial->BlendMode = BLEND_Masked;
@@ -1066,13 +1066,13 @@ UMaterialInterface* CreateSpeedTreeMaterial8(UObject* Parent, FString MaterialFu
 	if (GeomType != STG_Billboard)
 	{
 		VertexColor = NewObject<UMaterialExpressionVertexColor>(UnrealMaterial);
-		UnrealMaterial->Expressions.Add(VertexColor);
-		UnrealMaterial->AmbientOcclusion.Expression = VertexColor;
-		UnrealMaterial->AmbientOcclusion.Mask = VertexColor->GetOutputs()[0].Mask;
-		UnrealMaterial->AmbientOcclusion.MaskR = 1;
-		UnrealMaterial->AmbientOcclusion.MaskG = 0;
-		UnrealMaterial->AmbientOcclusion.MaskB = 0;
-		UnrealMaterial->AmbientOcclusion.MaskA = 0;
+		UnrealMaterial->GetExpressionCollection().AddExpression(VertexColor);
+		UnrealMaterialEditorOnly->AmbientOcclusion.Expression = VertexColor;
+		UnrealMaterialEditorOnly->AmbientOcclusion.Mask = VertexColor->GetOutputs()[0].Mask;
+		UnrealMaterialEditorOnly->AmbientOcclusion.MaskR = 1;
+		UnrealMaterialEditorOnly->AmbientOcclusion.MaskG = 0;
+		UnrealMaterialEditorOnly->AmbientOcclusion.MaskB = 0;
+		UnrealMaterialEditorOnly->AmbientOcclusion.MaskA = 0;
 	}
 
 	// diffuse and opacity mask
@@ -1083,17 +1083,17 @@ UMaterialInterface* CreateSpeedTreeMaterial8(UObject* Parent, FString MaterialFu
 			UMaterialExpressionConstant3Vector* ColorExpression = NewObject<UMaterialExpressionConstant3Vector>(UnrealMaterial);
 			GameEngine8::Vec4 Color = SpeedTreeMaterial.Maps()[0].Color();
 			ColorExpression->Constant = FLinearColor(FColor(Color.x * 255, Color.y * 255, Color.z * 255));
-			UnrealMaterial->Expressions.Add(ColorExpression);
-			UnrealMaterial->BaseColor.Expression = ColorExpression;
+			UnrealMaterial->GetExpressionCollection().AddExpression(ColorExpression);
+			UnrealMaterialEditorOnly->BaseColor.Expression = ColorExpression;
 
 			if (VertexColor != NULL)
 			{
-				UnrealMaterial->OpacityMask.Expression = VertexColor;
-				UnrealMaterial->OpacityMask.Mask = VertexColor->GetOutputs()[0].Mask;
-				UnrealMaterial->OpacityMask.MaskR = 0;
-				UnrealMaterial->OpacityMask.MaskG = 0;
-				UnrealMaterial->OpacityMask.MaskB = 0;
-				UnrealMaterial->OpacityMask.MaskA = 1;
+				UnrealMaterialEditorOnly->OpacityMask.Expression = VertexColor;
+				UnrealMaterialEditorOnly->OpacityMask.Mask = VertexColor->GetOutputs()[0].Mask;
+				UnrealMaterialEditorOnly->OpacityMask.MaskR = 0;
+				UnrealMaterialEditorOnly->OpacityMask.MaskG = 0;
+				UnrealMaterialEditorOnly->OpacityMask.MaskB = 0;
+				UnrealMaterialEditorOnly->OpacityMask.MaskA = 1;
 			}
 		}
 		else
@@ -1113,8 +1113,8 @@ UMaterialInterface* CreateSpeedTreeMaterial8(UObject* Parent, FString MaterialFu
 				UMaterialExpressionTextureSample* TextureExpression = NewObject<UMaterialExpressionTextureSample>(UnrealMaterial);
 				TextureExpression->Texture = DiffuseTexture;
 				TextureExpression->SamplerType = SAMPLERTYPE_Color;
-				UnrealMaterial->Expressions.Add(TextureExpression);
-				UnrealMaterial->BaseColor.Expression = TextureExpression;
+				UnrealMaterial->GetExpressionCollection().AddExpression(TextureExpression);
+				UnrealMaterialEditorOnly->BaseColor.Expression = TextureExpression;
 
 				if (VertexColor == NULL)
 				{
@@ -1122,7 +1122,7 @@ UMaterialInterface* CreateSpeedTreeMaterial8(UObject* Parent, FString MaterialFu
 					if (BillboardCrossfadeFunction)
 					{
 						UMaterialExpressionMaterialFunctionCall* BillboardCrossfade = NewObject<UMaterialExpressionMaterialFunctionCall>(UnrealMaterial);
-						UnrealMaterial->Expressions.Add(BillboardCrossfade);
+						UnrealMaterial->GetExpressionCollection().AddExpression(BillboardCrossfade);
 						BillboardCrossfade->MaterialFunction = BillboardCrossfadeFunction;
 						BillboardCrossfade->UpdateFromFunctionResource();
 
@@ -1132,19 +1132,19 @@ UMaterialInterface* CreateSpeedTreeMaterial8(UObject* Parent, FString MaterialFu
 						BillboardCrossfade->GetInput(0)->MaskG = 0;
 						BillboardCrossfade->GetInput(0)->MaskB = 0;
 						BillboardCrossfade->GetInput(0)->MaskA = 1;
-						UnrealMaterial->OpacityMask.Expression = BillboardCrossfade;
-						UnrealMaterial->OpacityMask.OutputIndex = 0.;
+						UnrealMaterialEditorOnly->OpacityMask.Expression = BillboardCrossfade;
+						UnrealMaterialEditorOnly->OpacityMask.OutputIndex = 0.;
 
 						UnrealMaterial->NumCustomizedUVs = 3;
-						UnrealMaterial->CustomizedUVs[2].Expression = BillboardCrossfade;
-						UnrealMaterial->CustomizedUVs[2].OutputIndex = 1;
+						UnrealMaterialEditorOnly->CustomizedUVs[2].Expression = BillboardCrossfade;
+						UnrealMaterialEditorOnly->CustomizedUVs[2].OutputIndex = 1;
 					}
 				}
 				else
 				{
 					// make mask with blend value and opacity
 					auto Multiply = NewObject<UMaterialExpressionMultiply>(UnrealMaterial);
-					UnrealMaterial->Expressions.Add(Multiply);
+					UnrealMaterial->GetExpressionCollection().AddExpression(Multiply);
 					Multiply->B.Expression = VertexColor;
 					Multiply->B.Mask = VertexColor->GetOutputs()[0].Mask;
 					Multiply->B.MaskR = 0;
@@ -1157,7 +1157,7 @@ UMaterialInterface* CreateSpeedTreeMaterial8(UObject* Parent, FString MaterialFu
 					Multiply->A.MaskG = 0;
 					Multiply->A.MaskB = 0;
 					Multiply->A.MaskA = 1;
-					UnrealMaterial->OpacityMask.Expression = Multiply;
+					UnrealMaterialEditorOnly->OpacityMask.Expression = Multiply;
 				}
 			}
 		}
@@ -1170,8 +1170,8 @@ UMaterialInterface* CreateSpeedTreeMaterial8(UObject* Parent, FString MaterialFu
 		{
 			UMaterialExpressionConstant* RoughExpression = NewObject<UMaterialExpressionConstant>(UnrealMaterial);
 			RoughExpression->R = SpeedTreeMaterial.Maps()[1].Color().w;
-			UnrealMaterial->Expressions.Add(RoughExpression);
-			UnrealMaterial->Roughness.Expression = RoughExpression;
+			UnrealMaterial->GetExpressionCollection().AddExpression(RoughExpression);
+			UnrealMaterialEditorOnly->Roughness.Expression = RoughExpression;
 		}
 		else
 		{
@@ -1191,10 +1191,10 @@ UMaterialInterface* CreateSpeedTreeMaterial8(UObject* Parent, FString MaterialFu
 				UMaterialExpressionTextureSample* TextureExpression = NewObject<UMaterialExpressionTextureSample>(UnrealMaterial);
 				TextureExpression->Texture = NormalTexture;
 				TextureExpression->SamplerType = SAMPLERTYPE_LinearColor;
-				UnrealMaterial->Expressions.Add(TextureExpression);
+				UnrealMaterial->GetExpressionCollection().AddExpression(TextureExpression);
 
 				UMaterialExpressionConstantBiasScale* BiasScale = NewObject<UMaterialExpressionConstantBiasScale>(UnrealMaterial);
-				UnrealMaterial->Expressions.Add(BiasScale);
+				UnrealMaterial->GetExpressionCollection().AddExpression(BiasScale);
 				BiasScale->Bias = -0.5f;
 				BiasScale->Scale = 2.0f;
 				BiasScale->Input.Expression = TextureExpression;
@@ -1205,7 +1205,7 @@ UMaterialInterface* CreateSpeedTreeMaterial8(UObject* Parent, FString MaterialFu
 					if (BillboardNormalFunction)
 					{
 						UMaterialExpressionMaterialFunctionCall* BillboardNormals = NewObject<UMaterialExpressionMaterialFunctionCall>(UnrealMaterial);
-						UnrealMaterial->Expressions.Add(BillboardNormals);
+						UnrealMaterial->GetExpressionCollection().AddExpression(BillboardNormals);
 						BillboardNormals->MaterialFunction = BillboardNormalFunction;
 						BillboardNormals->UpdateFromFunctionResource();
 
@@ -1216,23 +1216,23 @@ UMaterialInterface* CreateSpeedTreeMaterial8(UObject* Parent, FString MaterialFu
 						BillboardNormals->GetInput(0)->MaskB = 1;
 						BillboardNormals->GetInput(0)->MaskA = 0;
 
-						UnrealMaterial->Normal.Expression = BillboardNormals;
+						UnrealMaterialEditorOnly->Normal.Expression = BillboardNormals;
 						UnrealMaterial->bTangentSpaceNormal = false;
 					}
 				}
 
-				if (UnrealMaterial->Normal.Expression == NULL)
+				if (UnrealMaterialEditorOnly->Normal.Expression == NULL)
 				{
-					UnrealMaterial->Normal.Expression = BiasScale;
+					UnrealMaterialEditorOnly->Normal.Expression = BiasScale;
 				}
 
 				// roughness
-				UnrealMaterial->Roughness.Expression = TextureExpression;
-				UnrealMaterial->Roughness.Mask = TextureExpression->GetOutputs()[0].Mask;
-				UnrealMaterial->Roughness.MaskR = 0;
-				UnrealMaterial->Roughness.MaskG = 0;
-				UnrealMaterial->Roughness.MaskB = 0;
-				UnrealMaterial->Roughness.MaskA = 1;
+				UnrealMaterialEditorOnly->Roughness.Expression = TextureExpression;
+				UnrealMaterialEditorOnly->Roughness.Mask = TextureExpression->GetOutputs()[0].Mask;
+				UnrealMaterialEditorOnly->Roughness.MaskR = 0;
+				UnrealMaterialEditorOnly->Roughness.MaskG = 0;
+				UnrealMaterialEditorOnly->Roughness.MaskB = 0;
+				UnrealMaterialEditorOnly->Roughness.MaskA = 1;
 			}
 		}
 	}
@@ -1245,8 +1245,8 @@ UMaterialInterface* CreateSpeedTreeMaterial8(UObject* Parent, FString MaterialFu
 			UMaterialExpressionConstant3Vector* ColorExpression = NewObject<UMaterialExpressionConstant3Vector>(UnrealMaterial);
 			GameEngine8::Vec4 Color = SpeedTreeMaterial.Maps()[2].Color();
 			ColorExpression->Constant = FLinearColor(FColor(Color.x * 255, Color.y * 255, Color.z * 255));
-			UnrealMaterial->Expressions.Add(ColorExpression);
-			UnrealMaterial->SubsurfaceColor.Expression = ColorExpression;
+			UnrealMaterial->GetExpressionCollection().AddExpression(ColorExpression);
+			UnrealMaterialEditorOnly->SubsurfaceColor.Expression = ColorExpression;
 		}
 		else
 		{
@@ -1257,8 +1257,8 @@ UMaterialInterface* CreateSpeedTreeMaterial8(UObject* Parent, FString MaterialFu
 				UMaterialExpressionTextureSample* TextureExpression = NewObject<UMaterialExpressionTextureSample>(UnrealMaterial);
 				TextureExpression->Texture = SubsurfaceTexture;
 				TextureExpression->SamplerType = SAMPLERTYPE_Color;
-				UnrealMaterial->Expressions.Add(TextureExpression);
-				UnrealMaterial->SubsurfaceColor.Expression = TextureExpression;
+				UnrealMaterial->GetExpressionCollection().AddExpression(TextureExpression);
+				UnrealMaterialEditorOnly->SubsurfaceColor.Expression = TextureExpression;
 
 				UnrealMaterial->SetShadingModel(MSM_TwoSidedFoliage);
 			}
@@ -1273,8 +1273,8 @@ UMaterialInterface* CreateSpeedTreeMaterial8(UObject* Parent, FString MaterialFu
 		SpeedTreeExpression->WindType = WindType;
 		SpeedTreeExpression->BillboardThreshold = 1.0f; // billboards use crossfade technique now in v8
 		SpeedTreeExpression->GeometryType = GeomType;
-		UnrealMaterial->Expressions.Add(SpeedTreeExpression);
-		UnrealMaterial->WorldPositionOffset.Expression = SpeedTreeExpression;
+		UnrealMaterial->GetExpressionCollection().AddExpression(SpeedTreeExpression);
+		UnrealMaterialEditorOnly->WorldPositionOffset.Expression = SpeedTreeExpression;
 	}
 
 	LayoutMaterial(UnrealMaterial, true);
@@ -1386,7 +1386,7 @@ UMaterialInterface* CreateSpeedTreeMaterial9(UObject* Parent, FString MaterialFu
 	// material static setup
 	FStaticParameterSet StaticParameters;
 	UnrealMaterialInstance->GetStaticParameterValues(StaticParameters);
-	for (FStaticSwitchParameter& SwitchParameter : StaticParameters.StaticSwitchParameters)
+	for (FStaticSwitchParameter& SwitchParameter : StaticParameters.EditorOnly.StaticSwitchParameters)
 	{
 		if (SwitchParameter.ParameterInfo.Name == FName(TEXT("SharedEnable")))
 		{
diff --git a/Engine/Plugins/Editor/WorldPartitionHLODUtilities/Source/Private/HLODBuilderMeshApproximate.cpp b/Engine/Plugins/Editor/WorldPartitionHLODUtilities/Source/Private/HLODBuilderMeshApproximate.cpp
--- a/Engine/Plugins/Editor/WorldPartitionHLODUtilities/Source/Private/HLODBuilderMeshApproximate.cpp
+++ b/Engine/Plugins/Editor/WorldPartitionHLODUtilities/Source/Private/HLODBuilderMeshApproximate.cpp
@@ -182,7 +182,7 @@ TArray<UActorComponent*> UHLODBuilderMeshApproximate::Build(const FHLODBuildCont
 					SwitchParameter.ParameterInfo.Name = ParamName;
 					SwitchParameter.Value = true;
 					SwitchParameter.bOverride = true;
-					StaticParameterSet.StaticSwitchParameters.Add(SwitchParameter);
+					StaticParameterSet.EditorOnly.StaticSwitchParameters.Add(SwitchParameter);
 				}
 			};
 
diff --git a/Engine/Plugins/Enterprise/AxFImporter/Source/AxFImporter/Private/AxFImporter.cpp b/Engine/Plugins/Enterprise/AxFImporter/Source/AxFImporter/Private/AxFImporter.cpp
--- a/Engine/Plugins/Enterprise/AxFImporter/Source/AxFImporter/Private/AxFImporter.cpp
+++ b/Engine/Plugins/Enterprise/AxFImporter/Source/AxFImporter/Private/AxFImporter.cpp
@@ -343,7 +343,7 @@ public:
 
 			if (Material->IsA<UMaterial>())
 			{
-				Material->Expressions.Add(Expression);
+				Material->GetExpressionCollection().AddExpression(Expression);
 				Expression->Material = Material;
 			}
 			// MaterialFunction->FunctionExpressions.Add(NewExpression);
@@ -856,6 +856,8 @@ public:
 			auto BaseColor = Box(nullptr);
 			auto Specular = Box(nullptr);
 
+			UMaterialEditorOnlyData* MaterialEditorOnly = Material->GetEditorOnlyData();
+
 			if (ExpressionFresnel)
 			{
 				auto F0 = R(ExpressionFresnel);
@@ -878,7 +880,7 @@ public:
 						BaseColor = ExpressionSpecularColor;
 					}
 				}
-				Connect(Material->Metallic, F0); // just to make sure that F0=1 switches Metallic
+				Connect(MaterialEditorOnly->Metallic, F0); // just to make sure that F0=1 switches Metallic
 
 				auto ScaleFresnelToAccountForUESpecularScaling = Mul(Fresnel(F0), 1 / 0.08f); // UE scales Specular by 0.08 to compute F0
 				Specular = ExpressionSpecularColor ? Mul(ScaleFresnelToAccountForUESpecularScaling, ExpressionSpecularColor) : ScaleFresnelToAccountForUESpecularScaling;
@@ -893,8 +895,8 @@ public:
 				}
 			}
 
-			Connect(Material->Metallic, Constant(1.0f));
-			Connect(Material->Roughness, Constant(RoughnessFitted));
+			Connect(MaterialEditorOnly->Metallic, Constant(1.0f));
+			Connect(MaterialEditorOnly->Roughness, Constant(RoughnessFitted));
 
 			if (bHasBRDFColorsTexture || bHasBTFFlakes)
 			{
@@ -1003,18 +1005,18 @@ public:
 						Lerp(lh, hh, ThetaFWeight),
 						ThetaIWeight);
 
-					Connect(Material->EmissiveColor, FlakeResult);
+					Connect(MaterialEditorOnly->EmissiveColor, FlakeResult);
 				}
 			}
 
 			if (BaseColor)
 			{
-				Connect(Material->BaseColor, BaseColor);
+				Connect(MaterialEditorOnly->BaseColor, BaseColor);
 			}
 
 			if (Specular)
 			{
-				Connect(Material->Specular, Specular);
+				Connect(MaterialEditorOnly->Specular, Specular);
 			}
 
 			bool bHasClearcoat = ExpressionClearcoatIOR || ExpressionClearcoatNormal || ExpressionClearcoatColor || (CC_IOR != 1.0f);
@@ -1029,7 +1031,7 @@ public:
 				// Clearcoat normal goes to regular Normal input
 				if (ExpressionClearcoatNormal)
 				{
-					Connect(Material->Normal, ExpressionClearcoatNormal);
+					Connect(MaterialEditorOnly->Normal, ExpressionClearcoatNormal);
 				}
 
 				if (ExpressionNormal)
@@ -1038,13 +1040,13 @@ public:
 					Connect(BottomNormal->Input, ExpressionNormal);
 				}
 
-				Connect(Material->ClearCoatRoughness, Constant(0.0f));
+				Connect(MaterialEditorOnly->ClearCoatRoughness, Constant(0.0f));
 			}
 			else
 			{
 				if (ExpressionNormal)
 				{
-					Connect(Material->Normal, ExpressionNormal);
+					Connect(MaterialEditorOnly->Normal, ExpressionNormal);
 				}
 			}
 		}
@@ -1269,7 +1271,7 @@ public:
 			}
 
 			// UE squares UI roughness, so take root
-			Connect(Material->Roughness, Sqrt(SpecularRoughness));
+			Connect(Material->GetEditorOnlyData()->Roughness, Sqrt(SpecularRoughness));
 		}
 
 		void SetIsAnisotropic(bool InIsAnisotropic) override
@@ -1358,6 +1360,8 @@ private:
 			auto BaseColor = Box(nullptr);
 			auto Specular = Box(nullptr);
 
+			UMaterialEditorOnlyData* MaterialEditorOnly = Material->GetEditorOnlyData();
+
 			if (ExpressionFresnel)
 			{
 				auto F0 = R(ExpressionFresnel);
@@ -1380,7 +1384,7 @@ private:
 						BaseColor = ExpressionSpecularColor;
 					}
 				}
-				Connect(Material->Metallic, F0); // just to make sure that F0=1 switches Metallic
+				Connect(MaterialEditorOnly->Metallic, F0); // just to make sure that F0=1 switches Metallic
 
 				auto ScaleFresnelToAccountForUESpecularScaling = Mul(Fresnel(F0), 1 / 0.08f); // UE scales Specular by 0.08 to compute F0
 				Specular = ExpressionSpecularColor ? Mul(ScaleFresnelToAccountForUESpecularScaling, ExpressionSpecularColor) : ScaleFresnelToAccountForUESpecularScaling;
@@ -1397,12 +1401,12 @@ private:
 
 			if (BaseColor)
 			{
-				Connect(Material->BaseColor, BaseColor);
+				Connect(MaterialEditorOnly->BaseColor, BaseColor);
 			}
 
 			if (Specular)
 			{
-				Connect(Material->Specular, Specular);
+				Connect(MaterialEditorOnly->Specular, Specular);
 			}
 
 			bool bHasClearcoat = ExpressionClearcoatIOR || ExpressionClearcoatNormal || ExpressionClearcoatColor;
@@ -1414,7 +1418,7 @@ private:
 				// Clearcoat normal goes to regular Normal input
 				if (ExpressionClearcoatNormal)
 				{
-					Connect(Material->Normal, ExpressionClearcoatNormal);
+					Connect(MaterialEditorOnly->Normal, ExpressionClearcoatNormal);
 				}
 
 				if (ExpressionNormal)
@@ -1427,7 +1431,7 @@ private:
 			{
 				if (ExpressionNormal)
 				{
-					Connect(Material->Normal, ExpressionNormal);
+					Connect(MaterialEditorOnly->Normal, ExpressionNormal);
 				}
 			}
 
@@ -1448,7 +1452,7 @@ private:
 
 					Connect(DitherTemporalAA->FunctionInputs[0].Input, ExpressionAlpha);
 
-					Connect(Material->OpacityMask, DitherTemporalAA);
+					Connect(MaterialEditorOnly->OpacityMask, DitherTemporalAA);
 				}
 			}
 		}
@@ -1502,7 +1506,7 @@ private:
 			}
 
 			// UE squares UI roughness, so take root
-			Connect(Material->Roughness, Sqrt(SpecularRoughness));
+			Connect(Material->GetEditorOnlyData()->Roughness, Sqrt(SpecularRoughness));
 		}
 
 		void SetIsAnisotropic(bool InIsAnisotropic)
@@ -2506,20 +2510,21 @@ private:
 			return false;
 
 		UMaterial* Material = Cast<UMaterial>(OutMaterial);
-		{
-			Material->BaseColor.Expression = nullptr;
-			Material->EmissiveColor.Expression = nullptr;
-			Material->SubsurfaceColor.Expression = nullptr;
-			Material->Roughness.Expression = nullptr;
-			Material->Metallic.Expression = nullptr;
-			Material->Specular.Expression = nullptr;
-			Material->Opacity.Expression = nullptr;
-			Material->Refraction.Expression = nullptr;
-			Material->OpacityMask.Expression = nullptr;
-			Material->ClearCoat.Expression = nullptr;
-			Material->ClearCoatRoughness.Expression = nullptr;
-			Material->Normal.Expression = nullptr;
-			Material->Expressions.Empty();
+		UMaterialEditorOnlyData* MaterialEditorOnly = Material->GetEditorOnlyData();
+		{
+			MaterialEditorOnly->BaseColor.Expression = nullptr;
+			MaterialEditorOnly->EmissiveColor.Expression = nullptr;
+			MaterialEditorOnly->SubsurfaceColor.Expression = nullptr;
+			MaterialEditorOnly->Roughness.Expression = nullptr;
+			MaterialEditorOnly->Metallic.Expression = nullptr;
+			MaterialEditorOnly->Specular.Expression = nullptr;
+			MaterialEditorOnly->Opacity.Expression = nullptr;
+			MaterialEditorOnly->Refraction.Expression = nullptr;
+			MaterialEditorOnly->OpacityMask.Expression = nullptr;
+			MaterialEditorOnly->ClearCoat.Expression = nullptr;
+			MaterialEditorOnly->ClearCoatRoughness.Expression = nullptr;
+			MaterialEditorOnly->Normal.Expression = nullptr;
+			Material->GetExpressionCollection().Empty();
 		}
 
 		for (auto& CreatedMaterialPtr : CreatedMaterials)
diff --git a/Engine/Plugins/Enterprise/DataprepEditor/Source/DataprepCore/Private/DataprepCorePrivateUtils.cpp b/Engine/Plugins/Enterprise/DataprepEditor/Source/DataprepCore/Private/DataprepCorePrivateUtils.cpp
--- a/Engine/Plugins/Enterprise/DataprepEditor/Source/DataprepCore/Private/DataprepCorePrivateUtils.cpp
+++ b/Engine/Plugins/Enterprise/DataprepEditor/Source/DataprepCore/Private/DataprepCorePrivateUtils.cpp
@@ -206,7 +206,7 @@ void DataprepCorePrivateUtils::CompileMaterial(UMaterialInterface* MaterialInter
 			FStaticParameterSet StaticParameters;
 			ConstantMaterialInstance->GetStaticParameterValues( StaticParameters );
 
-			for ( FStaticSwitchParameter& Switch : StaticParameters.StaticSwitchParameters )
+			for ( FStaticSwitchParameter& Switch : StaticParameters.EditorOnly.StaticSwitchParameters )
 			{
 				if ( Switch.bOverride )
 				{
diff --git a/Engine/Plugins/Enterprise/DatasmithContent/Source/DatasmithContent/Private/ObjectTemplates/DatasmithMaterialInstanceTemplate.cpp b/Engine/Plugins/Enterprise/DatasmithContent/Source/DatasmithContent/Private/ObjectTemplates/DatasmithMaterialInstanceTemplate.cpp
--- a/Engine/Plugins/Enterprise/DatasmithContent/Source/DatasmithContent/Private/ObjectTemplates/DatasmithMaterialInstanceTemplate.cpp
+++ b/Engine/Plugins/Enterprise/DatasmithContent/Source/DatasmithContent/Private/ObjectTemplates/DatasmithMaterialInstanceTemplate.cpp
@@ -103,7 +103,7 @@ void FDatasmithStaticParameterSetTemplate::Apply( UMaterialInstanceConstant* Des
 			}
 		}
 
-		for ( FStaticSwitchParameter& DestinationSwitchParameter : DestinationStaticParameters.StaticSwitchParameters )
+		for ( FStaticSwitchParameter& DestinationSwitchParameter : DestinationStaticParameters.EditorOnly.StaticSwitchParameters )
 		{
 			if ( DestinationSwitchParameter.ParameterInfo.Name == It->Key )
 			{
@@ -132,9 +132,9 @@ void FDatasmithStaticParameterSetTemplate::Load( const UMaterialInstanceConstant
 	FStaticParameterSet SourceStaticParameters;
 	const_cast< UMaterialInstanceConstant& >( Source ).GetStaticParameterValues( SourceStaticParameters );
 
-	StaticSwitchParameters.Empty( SourceStaticParameters.StaticSwitchParameters.Num() );
+	StaticSwitchParameters.Empty( SourceStaticParameters.EditorOnly.StaticSwitchParameters.Num() );
 
-	for ( const FStaticSwitchParameter& SourceSwitch : SourceStaticParameters.StaticSwitchParameters )
+	for ( const FStaticSwitchParameter& SourceSwitch : SourceStaticParameters.EditorOnly.StaticSwitchParameters )
 	{
 		if ( !bOverridesOnly || SourceSwitch.bOverride )
 		{
diff --git a/Engine/Plugins/Enterprise/DatasmithImporter/Source/DatasmithImporter/Private/DatasmithMaterialExpressions.cpp b/Engine/Plugins/Enterprise/DatasmithImporter/Source/DatasmithImporter/Private/DatasmithMaterialExpressions.cpp
--- a/Engine/Plugins/Enterprise/DatasmithImporter/Source/DatasmithImporter/Private/DatasmithMaterialExpressions.cpp
+++ b/Engine/Plugins/Enterprise/DatasmithImporter/Source/DatasmithImporter/Private/DatasmithMaterialExpressions.cpp
@@ -122,11 +122,11 @@ namespace
 	{
 		if (UMaterial* Material = Cast< UMaterial >(MaterialOrFunction))
 		{
-			return Material->Expressions.Num();
+			return Material->GetExpressions().Num();
 		}
 		else if (UMaterialFunction* MaterialFunction = Cast< UMaterialFunction >(MaterialOrFunction))
 		{
-			return MaterialFunction->FunctionExpressions.Num();
+			return MaterialFunction->GetExpressions().Num();
 		}
 
 		return 0;
@@ -237,11 +237,11 @@ namespace
 
 UMaterialExpressionMakeMaterialAttributes* FDatasmithMaterialExpressions::FindOrAddAttributesFromMatFunc(UMaterialFunction* Func)
 {
-	for (int32 i = 0; i < Func->FunctionExpressions.Num(); i++)
+	for (UMaterialExpression* Expression : Func->GetExpressions())
 	{
-		if (Func->FunctionExpressions[i]->IsA< UMaterialExpressionMakeMaterialAttributes >())
+		if (Expression->IsA< UMaterialExpressionMakeMaterialAttributes >())
 		{
-			return StaticCast< UMaterialExpressionMakeMaterialAttributes* >(Func->FunctionExpressions[i]);
+			return StaticCast< UMaterialExpressionMakeMaterialAttributes* >(Expression);
 		}
 	}
 
@@ -254,11 +254,11 @@ UMaterialExpressionMakeMaterialAttributes* FDatasmithMaterialExpressions::FindOr
 
 UMaterialExpressionFunctionOutput* FDatasmithMaterialExpressions::FindOrAddOutputFromMatFunc(UMaterialFunction* Func)
 {
-	for (int32 i = 0; i < Func->FunctionExpressions.Num(); i++)
+	for (UMaterialExpression* Expression : Func->GetExpressions())
 	{
-		if ( Func->FunctionExpressions[i]->IsA< UMaterialExpressionFunctionOutput >() )
+		if (Expression->IsA< UMaterialExpressionFunctionOutput >() )
 		{
-			return StaticCast< UMaterialExpressionFunctionOutput* >( Func->FunctionExpressions[i] );
+			return StaticCast< UMaterialExpressionFunctionOutput* >(Expression);
 		}
 	}
 
@@ -289,15 +289,6 @@ void FDatasmithMaterialExpressions::GetSamplersRecursive(UMaterialExpression* Ex
 	}
 }
 
-void FDatasmithMaterialExpressions::GetSamplers(TArray<UMaterialExpression*>& Expressions, TArray<FExpressionInput*> ExpressionInputs, TArray<UMaterialExpressionTextureSample*>& TextureSamplers)
-{
-	for (int32 ExpressionIndex = 0; ExpressionIndex < Expressions.Num(); ExpressionIndex++)
-	{
-		UMaterialExpression* Expression = Expressions[ExpressionIndex];
-		GetSamplersRecursive(Expression, TextureSamplers);
-	}
-}
-
 void FDatasmithMaterialExpressions::GetTextureSamplersMaterial(UMaterialInterface* MaterialInterface, TArray<UMaterialExpressionTextureSample*>& TextureSamplers)
 {
 	UMaterial* Material = MaterialInterface->GetMaterial();
@@ -312,12 +303,13 @@ void FDatasmithMaterialExpressions::GetTextureSamplersMaterial(UMaterialInterfac
 		ExpressionInputs.Add(FirstLevelExpression);
 	}
 
-	if (Material->MaterialAttributes.Expression)
+	UMaterialEditorOnlyData* MaterialEditorOnly = Material->GetEditorOnlyData();
+	if (MaterialEditorOnly->MaterialAttributes.Expression)
 	{
-		ExpressionInputs.Add(&Material->MaterialAttributes);
+		ExpressionInputs.Add(&MaterialEditorOnly->MaterialAttributes);
 	}
 
-	GetSamplers(Material->Expressions, ExpressionInputs, TextureSamplers);
+	GetSamplers(Material->GetExpressions(), ExpressionInputs, TextureSamplers);
 }
 
 void FDatasmithMaterialExpressions::GetTextureSamplersFunc(UMaterialFunction* MaterialFunc, TArray<UMaterialExpressionTextureSample*>& TextureSamplers)
@@ -380,13 +372,13 @@ bool FDatasmithMaterialExpressions::MatOrFuncDelExpression(UObject* Object, UMat
 	if (Object->IsA< UMaterial >())
 	{
 		UMaterial* Mat = (UMaterial*)Object;
-		Mat->Expressions.Remove(ToBeConnected);
+		Mat->GetExpressionCollection().RemoveExpression(ToBeConnected);
 		return true;
 	}
 	else if (Object->IsA< UMaterialFunction >())
 	{
 		UMaterialFunction* Func = (UMaterialFunction*)Object;
-		Func->FunctionExpressions.Remove(ToBeConnected);
+		Func->GetExpressionCollection().RemoveExpression(ToBeConnected);
 		return true;
 	}
 	return false;
@@ -1845,11 +1837,12 @@ void FDatasmithMaterialExpressions::AddGlassNode(UMaterialExpression* RefleExpre
 	// usually fully white so specular gets invalidated
 	if (ActualMaterial)
 	{
-		ActualMaterial->BaseColor.Expression = Multiply06;
-		ActualMaterial->Specular.Expression = RefleExpression;
-		ActualMaterial->Metallic.Expression = RefleExpression;
-		ActualMaterial->Opacity.Expression = OneMinus(UnrealMaterial, Desaturated);
-		ActualMaterial->Refraction.Expression = RefractionIOR(IORRefra, UnrealMaterial);
+		UMaterialEditorOnlyData* ActualMaterialEditorOnly = ActualMaterial->GetEditorOnlyData();
+		ActualMaterialEditorOnly->BaseColor.Expression = Multiply06;
+		ActualMaterialEditorOnly->Specular.Expression = RefleExpression;
+		ActualMaterialEditorOnly->Metallic.Expression = RefleExpression;
+		ActualMaterialEditorOnly->Opacity.Expression = OneMinus(UnrealMaterial, Desaturated);
+		ActualMaterialEditorOnly->Refraction.Expression = RefractionIOR(IORRefra, UnrealMaterial);
 		ActualMaterial->BlendMode = EBlendMode::BLEND_Translucent;
 		ActualMaterial->TranslucencyLightingMode = ETranslucencyLightingMode::TLM_Surface;
 		ActualMaterial->TwoSided = true;
@@ -1934,10 +1927,12 @@ void FDatasmithMaterialExpressions::ModulateNormalAmount(UObject* UnrealMatOrFun
 		Attrib = FindOrAddAttributesFromMatFunc(Func);
 	}
 
+	UMaterialEditorOnlyData* UnrealMaterialEditorOnly = nullptr;
 	UMaterialExpression* ToBeConnected = nullptr;
 	if (UnrealMaterial)
 	{
-		ToBeConnected = UnrealMaterial->Normal.Expression;
+		UnrealMaterialEditorOnly = UnrealMaterial->GetEditorOnlyData();
+		ToBeConnected = UnrealMaterialEditorOnly->Normal.Expression;
 	}
 
 	if (Attrib)
@@ -1958,9 +1953,9 @@ void FDatasmithMaterialExpressions::ModulateNormalAmount(UObject* UnrealMatOrFun
 	ExpFuncCall->MaterialFunction = Flatten;
 	ExpFuncCall->UpdateFromFunctionResource();
 
-	if (UnrealMaterial)
+	if (UnrealMaterialEditorOnly)
 	{
-		UnrealMaterial->Normal.Expression = ExpFuncCall;
+		UnrealMaterialEditorOnly->Normal.Expression = ExpFuncCall;
 	}
 
 	if (Attrib)
@@ -2441,7 +2436,8 @@ UMaterialExpressionMaterialFunctionCall* FDatasmithMaterialExpressions::BlendFun
 			WeightExpression->ConnectExpression(ExpFuncCall->GetInput(2), 0);
 		}
 
-		UnrealMaterial->MaterialAttributes.Expression = ExpFuncCall;
+		UMaterialEditorOnlyData* UnrealMaterialEditorOnly = UnrealMaterial->GetEditorOnlyData();
+		UnrealMaterialEditorOnly->MaterialAttributes.Expression = ExpFuncCall;
 		return ExpFuncCall;
 	}
 	else
@@ -2507,15 +2503,17 @@ UMaterialInterface* FDatasmithMaterialExpressions::CreateDatasmithMaterial(UPack
 	else
 	{
 		UnrealMaterial = ExistingMaterial;
-		UnrealMaterial->BaseColor.Expression = nullptr;
-		UnrealMaterial->Metallic.Expression = nullptr;
-		UnrealMaterial->Specular.Expression = nullptr;
-		UnrealMaterial->Roughness.Expression = nullptr;
-		UnrealMaterial->Normal.Expression = nullptr;
-		UnrealMaterial->EmissiveColor.Expression = nullptr;
-		UnrealMaterial->Opacity.Expression = nullptr;
-		UnrealMaterial->OpacityMask.Expression = nullptr;
-		UnrealMaterial->Expressions.Empty();
+		UMaterialEditorOnlyData* UnrealMaterialEditorOnly = UnrealMaterial->GetEditorOnlyData();
+
+		UnrealMaterialEditorOnly->BaseColor.Expression = nullptr;
+		UnrealMaterialEditorOnly->Metallic.Expression = nullptr;
+		UnrealMaterialEditorOnly->Specular.Expression = nullptr;
+		UnrealMaterialEditorOnly->Roughness.Expression = nullptr;
+		UnrealMaterialEditorOnly->Normal.Expression = nullptr;
+		UnrealMaterialEditorOnly->EmissiveColor.Expression = nullptr;
+		UnrealMaterialEditorOnly->Opacity.Expression = nullptr;
+		UnrealMaterialEditorOnly->OpacityMask.Expression = nullptr;
+		UnrealMaterial->GetExpressionCollection().Empty();
 	}
 	UnrealMaterial->bUseMaterialAttributes = true;
 
@@ -2576,15 +2574,16 @@ UMaterialInterface* FDatasmithMaterialExpressions::CreateDatasmithMaterial(UPack
 	else
 	{
 		UnrealMaterial = ExistingMaterial;
-		UnrealMaterial->BaseColor.Expression = nullptr;
-		UnrealMaterial->Metallic.Expression = nullptr;
-		UnrealMaterial->Specular.Expression = nullptr;
-		UnrealMaterial->Roughness.Expression = nullptr;
-		UnrealMaterial->Normal.Expression = nullptr;
-		UnrealMaterial->EmissiveColor.Expression = nullptr;
-		UnrealMaterial->Opacity.Expression = nullptr;
-		UnrealMaterial->OpacityMask.Expression = nullptr;
-		UnrealMaterial->Expressions.Empty();
+		UMaterialEditorOnlyData* UnrealMaterialEditorOnly = UnrealMaterial->GetEditorOnlyData();
+		UnrealMaterialEditorOnly->BaseColor.Expression = nullptr;
+		UnrealMaterialEditorOnly->Metallic.Expression = nullptr;
+		UnrealMaterialEditorOnly->Specular.Expression = nullptr;
+		UnrealMaterialEditorOnly->Roughness.Expression = nullptr;
+		UnrealMaterialEditorOnly->Normal.Expression = nullptr;
+		UnrealMaterialEditorOnly->EmissiveColor.Expression = nullptr;
+		UnrealMaterialEditorOnly->Opacity.Expression = nullptr;
+		UnrealMaterialEditorOnly->OpacityMask.Expression = nullptr;
+		UnrealMaterial->GetExpressionCollection().Empty();
 	}
 
 	CreateDatasmithMaterialHelper(Package, ShaderElement, AssetsContext, UnrealMaterial);
@@ -2628,7 +2627,7 @@ UMaterialFunction* FDatasmithMaterialExpressions::CreateDatasmithMaterialFunc(UP
 	Attrib->EmissiveColor.Expression = nullptr;
 	Attrib->Opacity.Expression = nullptr;
 	Attrib->OpacityMask.Expression = nullptr;
-	UnrealMaterialFunc->FunctionExpressions.Empty();
+	UnrealMaterialFunc->GetExpressionCollection().Empty();
 
 	CreateDatasmithMaterialHelper(Package, ShaderElement, AssetsContext, UnrealMaterialFunc);
 
@@ -2661,15 +2660,16 @@ UMaterialInterface* FDatasmithMaterialExpressions::CreateDatasmithEnvironmentMat
 	else
 	{
 		UnrealMaterial = ExistingMaterial;
-		UnrealMaterial->BaseColor.Expression = nullptr;
-		UnrealMaterial->Metallic.Expression = nullptr;
-		UnrealMaterial->Specular.Expression = nullptr;
-		UnrealMaterial->Roughness.Expression = nullptr;
-		UnrealMaterial->Normal.Expression = nullptr;
-		UnrealMaterial->EmissiveColor.Expression = nullptr;
-		UnrealMaterial->Opacity.Expression = nullptr;
-		UnrealMaterial->OpacityMask.Expression = nullptr;
-		UnrealMaterial->Expressions.Empty();
+		UMaterialEditorOnlyData* UnrealMaterialEditorOnly = UnrealMaterial->GetEditorOnlyData();
+		UnrealMaterialEditorOnly->BaseColor.Expression = nullptr;
+		UnrealMaterialEditorOnly->Metallic.Expression = nullptr;
+		UnrealMaterialEditorOnly->Specular.Expression = nullptr;
+		UnrealMaterialEditorOnly->Roughness.Expression = nullptr;
+		UnrealMaterialEditorOnly->Normal.Expression = nullptr;
+		UnrealMaterialEditorOnly->EmissiveColor.Expression = nullptr;
+		UnrealMaterialEditorOnly->Opacity.Expression = nullptr;
+		UnrealMaterialEditorOnly->OpacityMask.Expression = nullptr;
+		UnrealMaterial->GetExpressionCollection().Empty();
 	}
 
 	Package->SetDirtyFlag(true);
diff --git a/Engine/Plugins/Enterprise/DatasmithImporter/Source/DatasmithImporter/Private/Utility/DatasmithImporterImpl.cpp b/Engine/Plugins/Enterprise/DatasmithImporter/Source/DatasmithImporter/Private/Utility/DatasmithImporterImpl.cpp
--- a/Engine/Plugins/Enterprise/DatasmithImporter/Source/DatasmithImporter/Private/Utility/DatasmithImporterImpl.cpp
+++ b/Engine/Plugins/Enterprise/DatasmithImporter/Source/DatasmithImporter/Private/Utility/DatasmithImporterImpl.cpp
@@ -383,7 +383,7 @@ void FDatasmithImporterImpl::CompileMaterial( UObject* Material, FMaterialUpdate
 			FStaticParameterSet StaticParameters;
 			ConstantMaterialInstance->GetStaticParameterValues( StaticParameters );
 
-			for ( FStaticSwitchParameter& Switch : StaticParameters.StaticSwitchParameters )
+			for ( FStaticSwitchParameter& Switch : StaticParameters.EditorOnly.StaticSwitchParameters )
 			{
 				if ( Switch.bOverride )
 				{
diff --git a/Engine/Plugins/Enterprise/DatasmithImporter/Source/DatasmithImporter/Public/DatasmithMaterialExpressions.h b/Engine/Plugins/Enterprise/DatasmithImporter/Source/DatasmithImporter/Public/DatasmithMaterialExpressions.h
--- a/Engine/Plugins/Enterprise/DatasmithImporter/Source/DatasmithImporter/Public/DatasmithMaterialExpressions.h
+++ b/Engine/Plugins/Enterprise/DatasmithImporter/Source/DatasmithImporter/Public/DatasmithMaterialExpressions.h
@@ -78,7 +78,16 @@ private:
 	static UMaterialExpressionFunctionOutput* FindOrAddOutputFromMatFunc(UMaterialFunction* Func);
 
 	static void GetSamplersRecursive(UMaterialExpression* Expression, TArray<UMaterialExpressionTextureSample*>& TextureSamplers);
-	static void GetSamplers(TArray<UMaterialExpression*>& Expressions, TArray<FExpressionInput*> ExpressionInputs, TArray<UMaterialExpressionTextureSample*>& TextureSamplers);
+
+	template<typename ExpressionsArrayType>
+	static void GetSamplers(const ExpressionsArrayType& Expressions, TArray<FExpressionInput*> ExpressionInputs, TArray<UMaterialExpressionTextureSample*>& TextureSamplers)
+	{
+		for (UMaterialExpression* Expression : Expressions)
+		{
+			GetSamplersRecursive(Expression, TextureSamplers);
+		}
+	}
+
 	static void GetTextureSamplersMaterial(UMaterialInterface* MaterialInterface, TArray<UMaterialExpressionTextureSample*>& TextureSamplers);
 	static void GetTextureSamplersFunc(UMaterialFunction* MaterialFunc, TArray<UMaterialExpressionTextureSample*>& TextureSamplers);
 	static void GetTextureSamplers(UObject* UnrealMatOrFunc, TArray<UMaterialExpressionTextureSample*>& TextureSamplers);
diff --git a/Engine/Plugins/Enterprise/MDLImporter/Source/MDLImporter/Private/MDLImporter.cpp b/Engine/Plugins/Enterprise/MDLImporter/Source/MDLImporter/Private/MDLImporter.cpp
--- a/Engine/Plugins/Enterprise/MDLImporter/Source/MDLImporter/Private/MDLImporter.cpp
+++ b/Engine/Plugins/Enterprise/MDLImporter/Source/MDLImporter/Private/MDLImporter.cpp
@@ -78,20 +78,22 @@ namespace MDLImporterImpl
 
 	void ClearMaterial(UMaterial* Material)
 	{
-		Material->BaseColor.Expression          = nullptr;
-		Material->EmissiveColor.Expression      = nullptr;
-		Material->SubsurfaceColor.Expression    = nullptr;
-		Material->Roughness.Expression          = nullptr;
-		Material->Metallic.Expression           = nullptr;
-		Material->Specular.Expression           = nullptr;
-		Material->Opacity.Expression            = nullptr;
-		Material->Refraction.Expression         = nullptr;
-		Material->OpacityMask.Expression        = nullptr;
-		Material->ClearCoat.Expression          = nullptr;
-		Material->ClearCoatRoughness.Expression = nullptr;
-		Material->Normal.Expression             = nullptr;
-
-		Material->Expressions.Empty();
+		UMaterialEditorOnlyData* MaterialEditorOnly = Material->GetEditorOnlyData();
+
+		MaterialEditorOnly->BaseColor.Expression          = nullptr;
+		MaterialEditorOnly->EmissiveColor.Expression      = nullptr;
+		MaterialEditorOnly->SubsurfaceColor.Expression    = nullptr;
+		MaterialEditorOnly->Roughness.Expression          = nullptr;
+		MaterialEditorOnly->Metallic.Expression           = nullptr;
+		MaterialEditorOnly->Specular.Expression           = nullptr;
+		MaterialEditorOnly->Opacity.Expression            = nullptr;
+		MaterialEditorOnly->Refraction.Expression         = nullptr;
+		MaterialEditorOnly->OpacityMask.Expression        = nullptr;
+		MaterialEditorOnly->ClearCoat.Expression          = nullptr;
+		MaterialEditorOnly->ClearCoatRoughness.Expression = nullptr;
+		MaterialEditorOnly->Normal.Expression             = nullptr;
+
+		Material->GetExpressionCollection().Empty();
 	}
 }
 
diff --git a/Engine/Plugins/Enterprise/MDLImporter/Source/MDLImporter/Private/MDLMaterialFactory.cpp b/Engine/Plugins/Enterprise/MDLImporter/Source/MDLImporter/Private/MDLMaterialFactory.cpp
--- a/Engine/Plugins/Enterprise/MDLImporter/Source/MDLImporter/Private/MDLMaterialFactory.cpp
+++ b/Engine/Plugins/Enterprise/MDLImporter/Source/MDLImporter/Private/MDLMaterialFactory.cpp
@@ -34,6 +34,8 @@ namespace MDLImporterImpl
 
 		const FMDLMaterialSelector::EMaterialType MaterialType = MaterialSelector.GetMaterialType(MdlMaterial);
 
+		UMaterialEditorOnlyData& MaterialEditorOnly = *Material.GetEditorOnlyData();
+
 		// set material settings
 		Material.bTangentSpaceNormal = true;
 		switch (MaterialType)
@@ -121,12 +123,12 @@ namespace MDLImporterImpl
 		{
 			if (MaterialType != FMDLMaterialSelector::EMaterialType::Translucent)
 				// TODO: find a nicer way to handle this cases
-				Connect(Material.BaseColor, MdlMaterial.BaseColor.ExpressionData);
-			Connect(Material.EmissiveColor, MdlMaterial.Emission.ExpressionData);
-			Connect(Material.SubsurfaceColor, MdlMaterial.Scattering.ExpressionData);
-			Connect(Material.Roughness, MdlMaterial.Roughness.ExpressionData);
-			Connect(Material.Metallic, MdlMaterial.Metallic.ExpressionData);
-			Connect(Material.Specular, MdlMaterial.Specular.ExpressionData);
+				Connect(MaterialEditorOnly.BaseColor, MdlMaterial.BaseColor.ExpressionData);
+			Connect(MaterialEditorOnly.EmissiveColor, MdlMaterial.Emission.ExpressionData);
+			Connect(MaterialEditorOnly.SubsurfaceColor, MdlMaterial.Scattering.ExpressionData);
+			Connect(MaterialEditorOnly.Roughness, MdlMaterial.Roughness.ExpressionData);
+			Connect(MaterialEditorOnly.Metallic, MdlMaterial.Metallic.ExpressionData);
+			Connect(MaterialEditorOnly.Specular, MdlMaterial.Specular.ExpressionData);
 
 			if (MaterialType == FMDLMaterialSelector::EMaterialType::Translucent)
 			{
@@ -134,26 +136,26 @@ namespace MDLImporterImpl
 			}
 			else if (MaterialType == FMDLMaterialSelector::EMaterialType::Masked)
 			{
-				Connect(Material.OpacityMask, MdlMaterial.Opacity.ExpressionData);
+				Connect(MaterialEditorOnly.OpacityMask, MdlMaterial.Opacity.ExpressionData);
 			}
 
-			Connect(Material.ClearCoat, MdlMaterial.Clearcoat.Weight.ExpressionData);
-			Connect(Material.ClearCoatRoughness, MdlMaterial.Clearcoat.Roughness.ExpressionData);
+			Connect(MaterialEditorOnly.ClearCoat, MdlMaterial.Clearcoat.Weight.ExpressionData);
+			Connect(MaterialEditorOnly.ClearCoatRoughness, MdlMaterial.Clearcoat.Roughness.ExpressionData);
 			if (UnderClearcoatNormal)
 			{
-				Connect(Material.Normal, MdlMaterial.Clearcoat.Normal.ExpressionData);
+				Connect(MaterialEditorOnly.Normal, MdlMaterial.Clearcoat.Normal.ExpressionData);
 				Connect(UnderClearcoatNormal->Input, MdlMaterial.Normal.ExpressionData);
 			}
 			else
 			{
-				Connect(Material.Normal, MdlMaterial.Normal.ExpressionData);
+				Connect(MaterialEditorOnly.Normal, MdlMaterial.Normal.ExpressionData);
 			}
 		}
 
 		if (UnderClearcoatNormal && !UnderClearcoatNormal->Input.Expression)
 		{
 			// delete if unused
-			Material.Expressions.Remove(UnderClearcoatNormal);
+			Material.GetExpressionCollection().RemoveExpression(UnderClearcoatNormal);
 			UnderClearcoatNormal = nullptr;
 		}
 
diff --git a/Engine/Plugins/Enterprise/MDLImporter/Source/MDLImporter/Private/generator/ConstantExpressionFactory.cpp b/Engine/Plugins/Enterprise/MDLImporter/Source/MDLImporter/Private/generator/ConstantExpressionFactory.cpp
--- a/Engine/Plugins/Enterprise/MDLImporter/Source/MDLImporter/Private/generator/ConstantExpressionFactory.cpp
+++ b/Engine/Plugins/Enterprise/MDLImporter/Source/MDLImporter/Private/generator/ConstantExpressionFactory.cpp
@@ -183,7 +183,7 @@ namespace Generator
 		{
 			if (!ExpressionHandle.Expression->UserCount)
 			{
-				CurrentMaterial->Expressions.Remove(ExpressionHandle.GetMaterialExpression());
+				CurrentMaterial->GetExpressionCollection().RemoveExpression(ExpressionHandle.GetMaterialExpression());
 				ExpressionHandle.DestroyExpression();
 			}
 		}
diff --git a/Engine/Plugins/Enterprise/MDLImporter/Source/MDLImporter/Private/generator/GeneratorMaterialExpressions.cpp b/Engine/Plugins/Enterprise/MDLImporter/Source/MDLImporter/Private/generator/GeneratorMaterialExpressions.cpp
--- a/Engine/Plugins/Enterprise/MDLImporter/Source/MDLImporter/Private/generator/GeneratorMaterialExpressions.cpp
+++ b/Engine/Plugins/Enterprise/MDLImporter/Source/MDLImporter/Private/generator/GeneratorMaterialExpressions.cpp
@@ -770,8 +770,8 @@ namespace Generator
 		}
 	}
 
-	template <typename ExpressionType>
-	int32 CountExpressions(const TArray<class UMaterialExpression*>& Expressions)
+	template <typename ExpressionType, typename ExpressionArrayType>
+	int32 CountExpressions(const ExpressionArrayType& Expressions)
 	{
 		int32 Count = 0;
 		for (UMaterialExpression* Expression : Expressions)
@@ -1068,7 +1068,7 @@ namespace Generator
 		Expression->Id                               = FGuid::NewGuid();
 		Expression->InputName                        = *Name;
 		Expression->InputType                        = Type;
-		Expression->SortPriority = CountExpressions<UMaterialExpressionFunctionInput>(Cast<UMaterialFunction>(Parent)->FunctionExpressions);
+		Expression->SortPriority = CountExpressions<UMaterialExpressionFunctionInput>(Cast<UMaterialFunction>(Parent)->GetExpressions());
 		Expression->bCollapsed   = true;
 		return Expression;
 	}
@@ -1091,7 +1091,7 @@ namespace Generator
 		Expression->Id                                = FGuid::NewGuid();
 		Expression->OutputName                        = *Name;
 		CheckedConnect(Parent, Output, Expression->A);
-		Expression->SortPriority = CountExpressions<UMaterialExpressionFunctionOutput>(Cast<UMaterialFunction>(Parent)->FunctionExpressions);
+		Expression->SortPriority = CountExpressions<UMaterialExpressionFunctionOutput>(Cast<UMaterialFunction>(Parent)->GetExpressions());
 		Expression->bCollapsed   = true;
 		return Expression;
 	}
diff --git a/Engine/Plugins/Enterprise/MDLImporter/Source/MDLImporter/Private/generator/MaterialExpressionFactory.cpp b/Engine/Plugins/Enterprise/MDLImporter/Source/MDLImporter/Private/generator/MaterialExpressionFactory.cpp
--- a/Engine/Plugins/Enterprise/MDLImporter/Source/MDLImporter/Private/generator/MaterialExpressionFactory.cpp
+++ b/Engine/Plugins/Enterprise/MDLImporter/Source/MDLImporter/Private/generator/MaterialExpressionFactory.cpp
@@ -766,7 +766,7 @@ namespace Generator
 				check(FunctionName == TEXT("operator[](<0>[],int)"));
 
 				const int32 Index = (int32)Cast<UMaterialExpressionConstant>(Inputs[1].GetExpressionUnused())->R;
-				CurrentMaterial->Expressions.Remove(Inputs[1].GetExpressionUnused());
+				CurrentMaterial->GetExpressionCollection().RemoveExpression(Inputs[1].GetExpressionUnused());
 
 				return {NewMaterialExpressionComponentMask(CurrentMaterial, Inputs[0], 1 << Index)};
 			}
@@ -1205,12 +1205,12 @@ namespace Generator
 			if (UMaterialExpressionConstant* Index = Cast<UMaterialExpressionConstant>(InputExpression))
 			{
 				CoordinateIndex = Index->R;
-				CurrentMaterial->Expressions.Remove(Index);
+				CurrentMaterial->GetExpressionCollection().RemoveExpression(Index);
 			}
 			else if (UMaterialExpressionScalarParameter* ScalarExpression = Cast<UMaterialExpressionScalarParameter>(InputExpression))
 			{
 				CoordinateIndex = ScalarExpression->DefaultValue;
-				CurrentMaterial->Expressions.Remove(ScalarExpression);
+				CurrentMaterial->GetExpressionCollection().RemoveExpression(ScalarExpression);
 			}
 			else
 			{
diff --git a/Engine/Plugins/Enterprise/MDLImporter/Source/MDLImporter/Private/generator/MaterialExpressions.h b/Engine/Plugins/Enterprise/MDLImporter/Source/MDLImporter/Private/generator/MaterialExpressions.h
--- a/Engine/Plugins/Enterprise/MDLImporter/Source/MDLImporter/Private/generator/MaterialExpressions.h
+++ b/Engine/Plugins/Enterprise/MDLImporter/Source/MDLImporter/Private/generator/MaterialExpressions.h
@@ -414,11 +414,11 @@ namespace Generator
 
 		if (Parent->IsA<UMaterial>())
 		{
-			Cast<UMaterial>(Parent)->Expressions.Push(Expression);
+			Cast<UMaterial>(Parent)->GetExpressionCollection().AddExpression(Expression);
 		}
 		else if (Parent->IsA<UMaterialFunction>())
 		{
-			Cast<UMaterialFunction>(Parent)->FunctionExpressions.Push(Expression);
+			Cast<UMaterialFunction>(Parent)->GetExpressionCollection().AddExpression(Expression);
 		}
 
 		return Expression;
diff --git a/Engine/Plugins/Enterprise/MDLImporter/Source/MDLImporter/Private/generator/ParameterExpressionFactory.cpp b/Engine/Plugins/Enterprise/MDLImporter/Source/MDLImporter/Private/generator/ParameterExpressionFactory.cpp
--- a/Engine/Plugins/Enterprise/MDLImporter/Source/MDLImporter/Private/generator/ParameterExpressionFactory.cpp
+++ b/Engine/Plugins/Enterprise/MDLImporter/Source/MDLImporter/Private/generator/ParameterExpressionFactory.cpp
@@ -440,7 +440,7 @@ namespace Generator
 				{
 					if (Connection.GetConnectionType() == EConnectionType::Expression)
 					{
-						CurrentMaterial->Expressions.Remove(Connection.GetExpressionUnused());
+						CurrentMaterial->GetExpressionCollection().RemoveExpression(Connection.GetExpressionUnused());
 						Connection.DestroyExpression();
 					}
 				}
diff --git a/Engine/Plugins/Enterprise/MDLImporter/Source/MDLImporter/Private/material/BakedMaterialFactory.cpp b/Engine/Plugins/Enterprise/MDLImporter/Source/MDLImporter/Private/material/BakedMaterialFactory.cpp
--- a/Engine/Plugins/Enterprise/MDLImporter/Source/MDLImporter/Private/material/BakedMaterialFactory.cpp
+++ b/Engine/Plugins/Enterprise/MDLImporter/Source/MDLImporter/Private/material/BakedMaterialFactory.cpp
@@ -55,10 +55,12 @@ namespace Mat
 
 	void FBakedMaterialFactory::Create(const Mdl::FMaterial& MdlMaterial, const Mat::FParameterMap& Parameters, UMaterial& Material) const
 	{
+		UMaterialEditorOnlyData& MaterialEditorOnly = *Material.GetEditorOnlyData();
+
 		// get under clear coat output
 		UMaterialExpressionClearCoatNormalCustomOutput* UnderClearCoat = nullptr;
 		{
-			TObjectPtr<UMaterialExpression>* Found = Material.Expressions.FindByPredicate(
+			const TObjectPtr<UMaterialExpression>* Found = Material.GetExpressions().FindByPredicate(
 			    [](const UMaterialExpression* Expr) { return Expr->IsA<UMaterialExpressionClearCoatNormalCustomOutput>(); });
 			if (Found)
 				UnderClearCoat = Cast<UMaterialExpressionClearCoatNormalCustomOutput>(*Found);
@@ -86,26 +88,26 @@ namespace Mat
 		Mat::FMapConnecter MapConnecter(Parameters, FunctionLoader, Tiling, Material);
 
 		// color
-		MapConnecter.ConnectParameterMap(Material.BaseColor, TEXT("Color"), EMaterialParameter::BaseColor);
+		MapConnecter.ConnectParameterMap(MaterialEditorOnly.BaseColor, TEXT("Color"), EMaterialParameter::BaseColor);
 
 		// brdf
-		MapConnecter.ConnectParameterMap(Material.Metallic, TEXT("BRDF"), EMaterialParameter::Metallic, false);
-		MapConnecter.ConnectParameterMap(Material.Specular, TEXT("BRDF"), EMaterialParameter::Specular, false);
-		MapConnecter.ConnectParameterMap(Material.Roughness, TEXT("BRDF"), EMaterialParameter::Roughness, false);
+		MapConnecter.ConnectParameterMap(MaterialEditorOnly.Metallic, TEXT("BRDF"), EMaterialParameter::Metallic, false);
+		MapConnecter.ConnectParameterMap(MaterialEditorOnly.Specular, TEXT("BRDF"), EMaterialParameter::Specular, false);
+		MapConnecter.ConnectParameterMap(MaterialEditorOnly.Roughness, TEXT("BRDF"), EMaterialParameter::Roughness, false);
 
 		// clear coat
 		if (UnderClearCoat)
 		{
 			check(Material.GetShadingModels().HasShadingModel(EMaterialShadingModel::MSM_ClearCoat));
-			MapConnecter.ConnectParameterMap(Material.ClearCoat, TEXT("Clear Coat"), EMaterialParameter::ClearCoatWeight, false);
-			MapConnecter.ConnectParameterMap(Material.ClearCoatRoughness, TEXT("Clear Coat"), EMaterialParameter::ClearCoatRoughness, false);
+			MapConnecter.ConnectParameterMap(MaterialEditorOnly.ClearCoat, TEXT("Clear Coat"), EMaterialParameter::ClearCoatWeight, false);
+			MapConnecter.ConnectParameterMap(MaterialEditorOnly.ClearCoatRoughness, TEXT("Clear Coat"), EMaterialParameter::ClearCoatRoughness, false);
 
-			MapConnecter.ConnectNormalMap(Material.Normal, TEXT("Clear Coat"), EMaterialParameter::ClearCoatNormalMap);
+			MapConnecter.ConnectNormalMap(MaterialEditorOnly.Normal, TEXT("Clear Coat"), EMaterialParameter::ClearCoatNormalMap);
 			MapConnecter.ConnectNormalMap(UnderClearCoat->Input, TEXT("Normal"), EMaterialParameter::NormalMap);
 		}
 		else
 		{
-			MapConnecter.ConnectNormalMap(Material.Normal, TEXT("Normal"), EMaterialParameter::NormalMap);
+			MapConnecter.ConnectNormalMap(MaterialEditorOnly.Normal, TEXT("Normal"), EMaterialParameter::NormalMap);
 
 			MapConnecter.DeleteExpressionMap(EMaterialParameter::ClearCoatWeight);
 			MapConnecter.DeleteExpressionMap(EMaterialParameter::ClearCoatRoughness);
@@ -115,7 +117,7 @@ namespace Mat
 
 		// subsurface
 		if (Material.GetShadingModels().HasShadingModel(EMaterialShadingModel::MSM_Subsurface))
-			MapConnecter.ConnectParameterMap(Material.SubsurfaceColor, TEXT("Color"), EMaterialParameter::SubSurfaceColor);
+			MapConnecter.ConnectParameterMap(MaterialEditorOnly.SubsurfaceColor, TEXT("Color"), EMaterialParameter::SubSurfaceColor);
 		else
 			MapConnecter.DeleteExpressionMap(EMaterialParameter::SubSurfaceColor);
 
@@ -123,14 +125,14 @@ namespace Mat
 		if (Parameters.Contains(EMaterialParameter::EmissionColorMap))
 		{
 			UMaterialExpression* StrengthParameter = ExpressionMakeFLoat3(Parameters[EMaterialParameter::EmissionStrength], FunctionLoader, Material);
-			MapConnecter.ConnectParameterMap(Material.EmissiveColor, TEXT("Emission"), EMaterialParameter::EmissionColor, true, StrengthParameter);
+			MapConnecter.ConnectParameterMap(MaterialEditorOnly.EmissiveColor, TEXT("Emission"), EMaterialParameter::EmissionColor, true, StrengthParameter);
 		}
 		else if (Parameters.Contains(EMaterialParameter::EmissionColor))
 		{
 			UMaterialExpression* ColorParameter    = Parameters[EMaterialParameter::EmissionColor];
 			UMaterialExpression* StrengthParameter = Parameters[EMaterialParameter::EmissionStrength];
 
-			Generator::Connect(Material.EmissiveColor, Generator::NewMaterialExpressionMultiply(&Material, {ColorParameter, StrengthParameter}));
+			Generator::Connect(MaterialEditorOnly.EmissiveColor, Generator::NewMaterialExpressionMultiply(&Material, {ColorParameter, StrengthParameter}));
 			Generator::SetMaterialExpressionGroup(TEXT("Emission"), ColorParameter);
 			Generator::SetMaterialExpressionGroup(TEXT("Emission"), StrengthParameter);
 		}
@@ -149,7 +151,7 @@ namespace Mat
 	{
 		FString               Name;
 
-		TObjectPtr<UMaterialExpression>* Found = Material.Expressions.FindByPredicate(
+		const TObjectPtr<UMaterialExpression>* Found = Material.GetExpressions().FindByPredicate(
 			[&Name](UMaterialExpression* Expression)  //
 		{
 			if (UMaterialExpressionScalarParameter* ScalarParameter = Cast<UMaterialExpressionScalarParameter>(Expression))
@@ -166,7 +168,7 @@ namespace Mat
 			return Name.Find(TEXT("Tiling Factor")) != INDEX_NONE;
 		});
 
-		TObjectPtr<UMaterialExpression>* FoundU = Material.Expressions.FindByPredicate(
+		const TObjectPtr<UMaterialExpression>* FoundU = Material.GetExpressions().FindByPredicate(
 		    [&Name](UMaterialExpression* Expression)  //
 		    {
 			    if (UMaterialExpressionScalarParameter* ScalarParameter = Cast<UMaterialExpressionScalarParameter>(Expression))
@@ -183,7 +185,7 @@ namespace Mat
 			    return Name.Find(TEXT("U Tiling")) != INDEX_NONE;
 		    });
 
-		TObjectPtr<UMaterialExpression>* FoundV = Material.Expressions.FindByPredicate(
+		const TObjectPtr<UMaterialExpression>* FoundV = Material.GetExpressions().FindByPredicate(
 			[&Name](UMaterialExpression* Expression)  //
 		{
 			if (UMaterialExpressionScalarParameter* ScalarParameter = Cast<UMaterialExpressionScalarParameter>(Expression))
diff --git a/Engine/Plugins/Enterprise/MDLImporter/Source/MDLImporter/Private/material/CarpaintMaterialFactory.cpp b/Engine/Plugins/Enterprise/MDLImporter/Source/MDLImporter/Private/material/CarpaintMaterialFactory.cpp
--- a/Engine/Plugins/Enterprise/MDLImporter/Source/MDLImporter/Private/material/CarpaintMaterialFactory.cpp
+++ b/Engine/Plugins/Enterprise/MDLImporter/Source/MDLImporter/Private/material/CarpaintMaterialFactory.cpp
@@ -66,7 +66,7 @@ namespace Mat
 		// get under clear coat output
 		UMaterialExpressionClearCoatNormalCustomOutput* UnderClearCoat;
 		{
-			auto Found = Material.Expressions.FindByPredicate(
+			auto Found = Material.GetExpressions().FindByPredicate(
 			    [](const UMaterialExpression* Expr) { return Expr->IsA<UMaterialExpressionClearCoatNormalCustomOutput>(); });
 			check(Found);
 			UnderClearCoat = Cast<UMaterialExpressionClearCoatNormalCustomOutput>(*Found);
@@ -81,9 +81,11 @@ namespace Mat
 		// change some parameter names
 		SetTextureParameterName(TEXT("Color Table"), Parameters[EMaterialParameter::BaseColorMap]);
 
+		UMaterialEditorOnlyData& MaterialEditorOnly = *Material.GetEditorOnlyData();
+
 		// clear coat
-		ConnectParameter(Material.ClearCoat, TEXT("Clear Coat"), Parameters[EMaterialParameter::ClearCoatWeight]);
-		ConnectParameter(Material.ClearCoatRoughness, TEXT("Clear Coat"), Parameters[EMaterialParameter::ClearCoatRoughness]);
+		ConnectParameter(MaterialEditorOnly.ClearCoat, TEXT("Clear Coat"), Parameters[EMaterialParameter::ClearCoatWeight]);
+		ConnectParameter(MaterialEditorOnly.ClearCoatRoughness, TEXT("Clear Coat"), Parameters[EMaterialParameter::ClearCoatRoughness]);
 		// light direction
 		Expression = Generator::NewMaterialExpressionVectorParameter(&Material, TEXT("Light Direction"), FLinearColor(0.5f, 0.25f, 1.f));
 		Generator::SetMaterialExpressionGroup(TEXT("Other"), Expression);
@@ -105,7 +107,7 @@ namespace Mat
 			Inputs.Emplace(Generator::FMaterialExpressionConnection(FlakesColor.Get<0>(), FlakesColor.Get<1>()));
 		else
 			Inputs.Emplace(Generator::NewMaterialExpressionConstant(&Material, 0.f));
-		Generator::Connect(Material.BaseColor, Generator::NewMaterialExpressionFunctionCall(&Material, MapCall, Inputs));
+		Generator::Connect(MaterialEditorOnly.BaseColor, Generator::NewMaterialExpressionFunctionCall(&Material, MapCall, Inputs));
 
 		MapCall = &FunctionLoader.Get(Generator::ECommonFunction::NormalMap);
 
@@ -116,7 +118,7 @@ namespace Mat
 			AddFunctionInput(Inputs, TEXT("Clear Coat"), Parameters[EMaterialParameter::ClearCoatNormalMap]);
 			Inputs.Emplace(Tiling);
 			AddFunctionInput(Inputs, TEXT("Clear Coat"), Parameters[EMaterialParameter::ClearCoatNormalStrength]);
-			Generator::Connect(Material.Normal, Generator::NewMaterialExpressionFunctionCall(&Material, MapCall, Inputs));
+			Generator::Connect(MaterialEditorOnly.Normal, Generator::NewMaterialExpressionFunctionCall(&Material, MapCall, Inputs));
 		}
 
 		// under-clearcoat normal
@@ -130,9 +132,9 @@ namespace Mat
 		if (!FlakesColor.Get<0>())
 		{
 			// brdf
-			ConnectParameter(Material.Metallic, TEXT("BRDF"), Parameters[EMaterialParameter::Metallic]);
-			ConnectParameter(Material.Specular, TEXT("BRDF"), Parameters[EMaterialParameter::Specular]);
-			ConnectParameter(Material.Roughness, TEXT("BRDF"), Parameters[EMaterialParameter::Roughness]);
+			ConnectParameter(MaterialEditorOnly.Metallic, TEXT("BRDF"), Parameters[EMaterialParameter::Metallic]);
+			ConnectParameter(MaterialEditorOnly.Specular, TEXT("BRDF"), Parameters[EMaterialParameter::Specular]);
+			ConnectParameter(MaterialEditorOnly.Roughness, TEXT("BRDF"), Parameters[EMaterialParameter::Roughness]);
 		}
 	}
 
@@ -164,21 +166,23 @@ namespace Mat
 
 		Generator::SetMaterialExpressionGroup(TEXT("Flakes"), Parameters[EMaterialParameter::CarFlakesMap]);
 
+		UMaterialEditorOnlyData& MaterialEditorOnly = *Material.GetEditorOnlyData();
+
 		// brdf
 		Inputs.Empty();
 		Inputs.Emplace(FunctionCall, 1);
 		Inputs.Emplace(Parameters[EMaterialParameter::Metallic]);
-		Generator::Connect(Material.Metallic,
+		Generator::Connect(MaterialEditorOnly.Metallic,
 		                   Generator::NewMaterialExpressionSaturate(&Material, Generator::NewMaterialExpressionAdd(&Material, Inputs)));
 		Inputs.Empty();
 		Inputs.Emplace(FunctionCall, 3);
 		Inputs.Emplace(Parameters[EMaterialParameter::Roughness]);
-		Generator::Connect(Material.Roughness,
+		Generator::Connect(MaterialEditorOnly.Roughness,
 		                   Generator::NewMaterialExpressionSaturate(&Material, Generator::NewMaterialExpressionAdd(&Material, Inputs)));
 		Inputs.Empty();
 		Inputs.Emplace(FunctionCall, 2);
 		Inputs.Emplace(Parameters[EMaterialParameter::Specular]);
-		Generator::Connect(Material.Specular,
+		Generator::Connect(MaterialEditorOnly.Specular,
 		                   Generator::NewMaterialExpressionSaturate(&Material, Generator::NewMaterialExpressionAdd(&Material, Inputs)));
 
 		Generator::SetMaterialExpressionGroup(TEXT("BRDF"), Parameters[EMaterialParameter::Metallic]);
diff --git a/Engine/Plugins/Enterprise/MDLImporter/Source/MDLImporter/Private/material/MapConnecter.h b/Engine/Plugins/Enterprise/MDLImporter/Source/MDLImporter/Private/material/MapConnecter.h
--- a/Engine/Plugins/Enterprise/MDLImporter/Source/MDLImporter/Private/material/MapConnecter.h
+++ b/Engine/Plugins/Enterprise/MDLImporter/Source/MDLImporter/Private/material/MapConnecter.h
@@ -25,7 +25,7 @@ namespace Mat
 				return;
 
 			UMaterialExpression* Expression = Parameters[Parameter];
-			Material.Expressions.Remove(Expression);
+			Material.GetExpressionCollection().RemoveExpression(Expression);
 		}
 
 		void DeleteExpressionMap(EMaterialParameter Parameter)
diff --git a/Engine/Plugins/Enterprise/MDLImporter/Source/MDLImporter/Private/material/TranslucentMaterialFactory.cpp b/Engine/Plugins/Enterprise/MDLImporter/Source/MDLImporter/Private/material/TranslucentMaterialFactory.cpp
--- a/Engine/Plugins/Enterprise/MDLImporter/Source/MDLImporter/Private/material/TranslucentMaterialFactory.cpp
+++ b/Engine/Plugins/Enterprise/MDLImporter/Source/MDLImporter/Private/material/TranslucentMaterialFactory.cpp
@@ -38,25 +38,27 @@ namespace Mat
 		UMaterialFunction*   MapCall    = nullptr;
 		UMaterialExpression* Expression = nullptr;
 
+		UMaterialEditorOnlyData& MaterialEditorOnly = *Material.GetEditorOnlyData();
+
 		// normal
-		UMaterialExpression* NormalExpression = MapConnecter.ConnectNormalMap(Material.Normal, TEXT("Normal"), EMaterialParameter::NormalMap);
+		UMaterialExpression* NormalExpression = MapConnecter.ConnectNormalMap(MaterialEditorOnly.Normal, TEXT("Normal"), EMaterialParameter::NormalMap);
 
 		// brdf
-		MapConnecter.ConnectParameterMap(Material.Metallic, TEXT("BRDF"), EMaterialParameter::Metallic, false);
-		MapConnecter.ConnectParameterMap(Material.Specular, TEXT("BRDF"), EMaterialParameter::Specular, false);
-		MapConnecter.ConnectParameterMap(Material.Roughness, TEXT("BRDF"), EMaterialParameter::Roughness, false);
+		MapConnecter.ConnectParameterMap(MaterialEditorOnly.Metallic, TEXT("BRDF"), EMaterialParameter::Metallic, false);
+		MapConnecter.ConnectParameterMap(MaterialEditorOnly.Specular, TEXT("BRDF"), EMaterialParameter::Specular, false);
+		MapConnecter.ConnectParameterMap(MaterialEditorOnly.Roughness, TEXT("BRDF"), EMaterialParameter::Roughness, false);
 
 		if (Material.BlendMode == BLEND_Masked && Parameters.Contains(EMaterialParameter::OpacityMap))
 		{
 			// masked
 
 			// color
-			MapConnecter.ConnectParameterMap(Material.BaseColor, TEXT("Color"), EMaterialParameter::BaseColor);
+			MapConnecter.ConnectParameterMap(MaterialEditorOnly.BaseColor, TEXT("Color"), EMaterialParameter::BaseColor);
 
 			UMaterialExpression* Opacity = MapConnecter.CreateMapCall(TEXT("Opacity"), EMaterialParameter::OpacityMap, false, nullptr);
 
 			MapCall = &FunctionLoader.Get(Generator::ECommonFunction::DitherTemporalAA);
-			Generator::Connect(Material.OpacityMask, Generator::NewMaterialExpressionFunctionCall(&Material, MapCall, {{Opacity, 0}, {nullptr, 0}}));
+			Generator::Connect(MaterialEditorOnly.OpacityMask, Generator::NewMaterialExpressionFunctionCall(&Material, MapCall, {{Opacity, 0}, {nullptr, 0}}));
 			MapConnecter.DeleteExpression(EMaterialParameter::AbsorptionColor);
 			MapConnecter.DeleteExpression(EMaterialParameter::IOR);
 		}
@@ -67,7 +69,7 @@ namespace Mat
 			// color
 			Expression = MapConnecter.CreateParameterMap(TEXT("Color"), EMaterialParameter::BaseColor);
 			Expression = Generator::NewMaterialExpressionMultiply(&Material, {Expression, Parameters[EMaterialParameter::Opacity]});
-			Generator::Connect(Material.BaseColor, Expression);
+			Generator::Connect(MaterialEditorOnly.BaseColor, Expression);
 
 			// path length expression
 			Expression = Generator::NewMaterialExpressionVectorParameter(&Material, TEXT("Local Bounds Adjustment"), FLinearColor(0.f, 0.f, 1.f));
@@ -89,11 +91,11 @@ namespace Mat
 			                                                  {Parameters[EMaterialParameter::IOR], 0},              //
 			                                                  {PathLengthExpression, 0},                             //
 			                                                  {NormalExpression, 0}});
-			Generator::Connect(Material.EmissiveColor, VolumeAbsorptionExpression);
+			Generator::Connect(MaterialEditorOnly.EmissiveColor, VolumeAbsorptionExpression);
 
 			// opacity
 			MapCall = &FunctionLoader.Get(Generator::ECommonFunction::TranslucentOpacity);
-			Generator::Connect(Material.Opacity, Generator::NewMaterialExpressionFunctionCall(&Material, MapCall, {}));
+			Generator::Connect(MaterialEditorOnly.Opacity, Generator::NewMaterialExpressionFunctionCall(&Material, MapCall, {}));
 		}
 	}
 }  // namespace Mat
diff --git a/Engine/Plugins/Experimental/Interchange/Runtime/Source/Import/Private/Material/InterchangeMaterialFactory.cpp b/Engine/Plugins/Experimental/Interchange/Runtime/Source/Import/Private/Material/InterchangeMaterialFactory.cpp
--- a/Engine/Plugins/Experimental/Interchange/Runtime/Source/Import/Private/Material/InterchangeMaterialFactory.cpp
+++ b/Engine/Plugins/Experimental/Interchange/Runtime/Source/Import/Private/Material/InterchangeMaterialFactory.cpp
@@ -328,7 +328,7 @@ void UInterchangeMaterialFactory::PreImportPreCompletedCallback(const FImportPre
 		//Update the samplers type in case the textures were changed during their PreImportPreCompletedCallback
 		if (UMaterial* ImportedMaterial = Cast<UMaterial>(ImportedMaterialInterface))
 		{
-			for (UMaterialExpression* Expression : ImportedMaterial->Expressions)
+			for (UMaterialExpression* Expression : ImportedMaterial->GetExpressions())
 			{
 				if (UMaterialExpressionTextureBase* TextureSample = Cast<UMaterialExpressionTextureBase>(Expression))
 				{
@@ -758,7 +758,7 @@ UMaterialExpression* UInterchangeMaterialFactory::CreateExpression(UMaterial* Ma
 	// Set the parameter name if the material expression has one (some material expressions don't inherit from UMaterialExpressionParameter, ie: UMaterialExpressionTextureSampleParameter
 	if (FNameProperty* Property = FindFProperty<FNameProperty>(MaterialExpression->GetClass(), GET_MEMBER_NAME_CHECKED(UMaterialExpressionParameter, ParameterName)))
 	{
-		*(Property->ContainerPtrToValuePtr<FName>(MaterialExpression)) = FName(*(ExpressionNode.GetDisplayLabel() + LexToString(Material->Expressions.Num())));
+		*(Property->ContainerPtrToValuePtr<FName>(MaterialExpression)) = FName(*(ExpressionNode.GetDisplayLabel() + LexToString(Material->GetExpressions().Num())));
 	}
 
 	ExpressionNode.ApplyAllCustomAttributeToObject(MaterialExpression);
diff --git a/Engine/Plugins/Experimental/VirtualHeightfieldMesh/Source/VirtualHeightfieldMesh/Private/HeightfieldMinMaxTextureMaterialExpression.h b/Engine/Plugins/Experimental/VirtualHeightfieldMesh/Source/VirtualHeightfieldMesh/Private/HeightfieldMinMaxTextureMaterialExpression.h
--- a/Engine/Plugins/Experimental/VirtualHeightfieldMesh/Source/VirtualHeightfieldMesh/Private/HeightfieldMinMaxTextureMaterialExpression.h
+++ b/Engine/Plugins/Experimental/VirtualHeightfieldMesh/Source/VirtualHeightfieldMesh/Private/HeightfieldMinMaxTextureMaterialExpression.h
@@ -19,10 +19,8 @@ class UMaterialExpressionHeightfieldMinMaxTexture : public UMaterialExpression
 	UPROPERTY(EditAnywhere, BlueprintReadWrite, Category = MaterialExpressionTexture)
 	UHeightfieldMinMaxTexture* MinMaxTexture;
 
-#if WITH_EDITORONLY_DATA
 	UPROPERTY(EditAnywhere, BlueprintReadWrite, Category = MaterialExpressionTexture)
 	TEnumAsByte<enum EMaterialSamplerType> SamplerType;
-#endif
 
 protected:
 #if WITH_EDITOR
diff --git a/Engine/Plugins/Importers/USDImporter/Source/USDUtilities/Private/USDShadeConversion.cpp b/Engine/Plugins/Importers/USDImporter/Source/USDUtilities/Private/USDShadeConversion.cpp
--- a/Engine/Plugins/Importers/USDImporter/Source/USDUtilities/Private/USDShadeConversion.cpp
+++ b/Engine/Plugins/Importers/USDImporter/Source/USDUtilities/Private/USDShadeConversion.cpp
@@ -1010,7 +1010,7 @@ namespace UE
 					FStaticParameterSet StaticParameters;
 					Constant->GetStaticParameterValues( StaticParameters );
 
-					for ( FStaticSwitchParameter& StaticSwitchParameter : StaticParameters.StaticSwitchParameters )
+					for ( FStaticSwitchParameter& StaticSwitchParameter : StaticParameters.EditorOnly.StaticSwitchParameters )
 					{
 						if ( StaticSwitchParameter.ParameterInfo.Name == ParameterName )
 						{
@@ -1288,13 +1288,13 @@ namespace UE
 						}
 						break;
 					case MP_Tangent:
-						if ( !Material.GetMaterial()->Tangent.IsConnected() && !Material.GetMaterial()->bUseMaterialAttributes )
+						if ( !Material.GetMaterial()->GetEditorOnlyData()->Tangent.IsConnected() && !Material.GetMaterial()->bUseMaterialAttributes )
 						{
 							continue;
 						}
 						break;
 					case MP_EmissiveColor:
-						if ( !Material.GetMaterial()->EmissiveColor.IsConnected() && !Material.GetMaterial()->bUseMaterialAttributes )
+						if ( !Material.GetMaterial()->GetEditorOnlyData()->EmissiveColor.IsConnected() && !Material.GetMaterial()->bUseMaterialAttributes )
 						{
 							continue;
 						}
@@ -1943,6 +1943,8 @@ bool UsdToUnreal::ConvertMaterial( const pxr::UsdShadeMaterial& UsdShadeMaterial
 		}
 	};
 
+	UMaterialEditorOnlyData* EditorOnly = Material.GetEditorOnlyData();
+
 	// Base color
 	{
 		const bool bIsNormalMap = false;
@@ -1950,8 +1952,8 @@ bool UsdToUnreal::ConvertMaterial( const pxr::UsdShadeMaterial& UsdShadeMaterial
 		{
 			if (UMaterialExpression* Expression = UsdShadeConversionImpl::GetExpressionForValue(Material, ParameterValue))
 			{
-				Material.BaseColor.Expression = Expression;
-				SetOutputIndex(ParameterValue, Material.BaseColor.OutputIndex);
+				EditorOnly->BaseColor.Expression = Expression;
+				SetOutputIndex(ParameterValue, EditorOnly->BaseColor.OutputIndex);
 
 				bHasMaterialInfo = true;
 			}
@@ -1965,8 +1967,8 @@ bool UsdToUnreal::ConvertMaterial( const pxr::UsdShadeMaterial& UsdShadeMaterial
 		{
 			if (UMaterialExpression* Expression = UsdShadeConversionImpl::GetExpressionForValue(Material, ParameterValue))
 			{
-				Material.EmissiveColor.Expression = Expression;
-				SetOutputIndex(ParameterValue, Material.EmissiveColor.OutputIndex);
+				EditorOnly->EmissiveColor.Expression = Expression;
+				SetOutputIndex(ParameterValue, EditorOnly->EmissiveColor.OutputIndex);
 
 				bHasMaterialInfo = true;
 			}
@@ -1978,8 +1980,8 @@ bool UsdToUnreal::ConvertMaterial( const pxr::UsdShadeMaterial& UsdShadeMaterial
 	{
 		if ( UMaterialExpression* Expression = UsdShadeConversionImpl::GetExpressionForValue( Material, ParameterValue ) )
 		{
-			Material.Metallic.Expression = Expression;
-			SetOutputIndex( ParameterValue, Material.Metallic.OutputIndex );
+			EditorOnly->Metallic.Expression = Expression;
+			SetOutputIndex( ParameterValue, EditorOnly->Metallic.OutputIndex );
 
 			bHasMaterialInfo = true;
 		}
@@ -1990,8 +1992,8 @@ bool UsdToUnreal::ConvertMaterial( const pxr::UsdShadeMaterial& UsdShadeMaterial
 	{
 		if ( UMaterialExpression* Expression = UsdShadeConversionImpl::GetExpressionForValue( Material, ParameterValue ) )
 		{
-			Material.Roughness.Expression = Expression;
-			SetOutputIndex( ParameterValue, Material.Roughness.OutputIndex );
+			EditorOnly->Roughness.Expression = Expression;
+			SetOutputIndex( ParameterValue, EditorOnly->Roughness.OutputIndex );
 
 			bHasMaterialInfo = true;
 		}
@@ -2002,8 +2004,8 @@ bool UsdToUnreal::ConvertMaterial( const pxr::UsdShadeMaterial& UsdShadeMaterial
 	{
 		if ( UMaterialExpression* Expression = UsdShadeConversionImpl::GetExpressionForValue( Material, ParameterValue ) )
 		{
-			Material.Opacity.Expression = Expression;
-			SetOutputIndex( ParameterValue, Material.Opacity.OutputIndex );
+			EditorOnly->Opacity.Expression = Expression;
+			SetOutputIndex( ParameterValue, EditorOnly->Opacity.OutputIndex );
 
 			Material.BlendMode = BLEND_Translucent;
 			bHasMaterialInfo = true;
@@ -2017,8 +2019,8 @@ bool UsdToUnreal::ConvertMaterial( const pxr::UsdShadeMaterial& UsdShadeMaterial
 		{
 			if ( UMaterialExpression* Expression = UsdShadeConversionImpl::GetExpressionForValue( Material, ParameterValue ) )
 			{
-				Material.Normal.Expression = Expression;
-				SetOutputIndex(ParameterValue, Material.Normal.OutputIndex);
+				EditorOnly->Normal.Expression = Expression;
+				SetOutputIndex(ParameterValue, EditorOnly->Normal.OutputIndex);
 
 				bHasMaterialInfo = true;
 			}
@@ -2031,8 +2033,8 @@ bool UsdToUnreal::ConvertMaterial( const pxr::UsdShadeMaterial& UsdShadeMaterial
 	{
 		if ( UMaterialExpression* Expression = UsdShadeConversionImpl::GetExpressionForValue( Material, ParameterValue ) )
 		{
-			Material.Refraction.Expression = Expression;
-			SetOutputIndex( ParameterValue, Material.Refraction.OutputIndex );
+			EditorOnly->Refraction.Expression = Expression;
+			SetOutputIndex( ParameterValue, EditorOnly->Refraction.OutputIndex );
 
 			bHasMaterialInfo = true;
 		}
diff --git a/Engine/Plugins/Media/MediaPlayerEditor/Source/MediaPlayerEditor/Private/Widgets/SMediaImage.cpp b/Engine/Plugins/Media/MediaPlayerEditor/Source/MediaPlayerEditor/Private/Widgets/SMediaImage.cpp
--- a/Engine/Plugins/Media/MediaPlayerEditor/Source/MediaPlayerEditor/Private/Widgets/SMediaImage.cpp
+++ b/Engine/Plugins/Media/MediaPlayerEditor/Source/MediaPlayerEditor/Private/Widgets/SMediaImage.cpp
@@ -29,6 +29,7 @@ void SMediaImage::Construct(const FArguments& InArgs, UTexture* InTexture)
 
 		if (Material != nullptr)
 		{
+			UMaterialEditorOnlyData* MaterialEditorOnly = Material->GetEditorOnlyData();
 			TextureSampler = NewObject<UMaterialExpressionTextureSample>(Material);
 			{
 				TextureSampler->Texture = InTexture;
@@ -36,7 +37,7 @@ void SMediaImage::Construct(const FArguments& InArgs, UTexture* InTexture)
 			}
 
 			FExpressionOutput& Output = TextureSampler->GetOutputs()[0];
-			FExpressionInput& Input = Material->EmissiveColor;
+			FExpressionInput& Input = MaterialEditorOnly->EmissiveColor;
 			{
 				Input.Expression = TextureSampler;
 				Input.Mask = Output.Mask;
@@ -46,7 +47,7 @@ void SMediaImage::Construct(const FArguments& InArgs, UTexture* InTexture)
 				Input.MaskA = Output.MaskA;
 			}
 
-			FExpressionInput& Opacity = Material->Opacity;
+			FExpressionInput& Opacity = MaterialEditorOnly->Opacity;
 			{
 				Opacity.Expression = TextureSampler;
 				Opacity.Mask = Output.Mask;
@@ -58,7 +59,7 @@ void SMediaImage::Construct(const FArguments& InArgs, UTexture* InTexture)
 
 			Material->BlendMode = BLEND_AlphaComposite;
 
-			Material->Expressions.Add(TextureSampler);
+			Material->GetExpressionCollection().AddExpression(TextureSampler);
 			Material->MaterialDomain = EMaterialDomain::MD_UI;
 			Material->PostEditChange();
 		}
diff --git a/Engine/Plugins/MeshPainting/Source/MeshPaintingToolset/Private/IMeshPaintComponentAdapter.cpp b/Engine/Plugins/MeshPainting/Source/MeshPaintingToolset/Private/IMeshPaintComponentAdapter.cpp
--- a/Engine/Plugins/MeshPainting/Source/MeshPaintingToolset/Private/IMeshPaintComponentAdapter.cpp
+++ b/Engine/Plugins/MeshPainting/Source/MeshPaintingToolset/Private/IMeshPaintComponentAdapter.cpp
@@ -47,12 +47,10 @@ void IMeshPaintComponentAdapter::DefaultQueryPaintableTextures(int32 MaterialInd
 		FPaintableTexture PaintableTexture;
 		// Find all the unique textures used in the top material level of the selected actor materials
 
-		const TArray<UMaterialExpression*>& Expressions = Material->GetMaterial()->Expressions;
-
 		// Only grab the textures from the top level of samples
-		for (auto ItExpressions = Expressions.CreateConstIterator(); ItExpressions; ItExpressions++)
+		for (UMaterialExpression* Expression : Material->GetMaterial()->GetExpressions())
 		{
-			UMaterialExpressionTextureBase* TextureBase = Cast<UMaterialExpressionTextureBase>(*ItExpressions);
+			UMaterialExpressionTextureBase* TextureBase = Cast<UMaterialExpressionTextureBase>(Expression);
 			if (TextureBase != NULL &&
 				TextureBase->Texture != NULL &&
 				!TextureBase->Texture->IsNormalMap() && 
@@ -66,7 +64,7 @@ void IMeshPaintComponentAdapter::DefaultQueryPaintableTextures(int32 MaterialInd
 
 				// Texture Samples can have UV's specified, check the first node for whether it has a custom UV channel set. 
 				// We only check the first as the Mesh paint mode does not support painting with UV's modified in the shader.
-				UMaterialExpressionTextureSample* TextureSample = Cast<UMaterialExpressionTextureSample>(*ItExpressions);
+				UMaterialExpressionTextureSample* TextureSample = Cast<UMaterialExpressionTextureSample>(Expression);
 				if (TextureSample != NULL)
 				{
 					UMaterialExpressionTextureCoordinate* TextureCoords = Cast<UMaterialExpressionTextureCoordinate>(TextureSample->Coordinates.Expression);
diff --git a/Engine/Plugins/Tests/EditorTests/Source/EditorTests/Private/UnrealEd/EditorBuildPromotionTests.cpp b/Engine/Plugins/Tests/EditorTests/Source/EditorTests/Private/UnrealEd/EditorBuildPromotionTests.cpp
--- a/Engine/Plugins/Tests/EditorTests/Source/EditorTests/Private/UnrealEd/EditorBuildPromotionTests.cpp
+++ b/Engine/Plugins/Tests/EditorTests/Source/EditorTests/Private/UnrealEd/EditorBuildPromotionTests.cpp
@@ -1495,9 +1495,9 @@ namespace BuildPromotionTestHelper
 
 				//Get the editor material
 				UMaterial* EditorMaterial = Cast<UMaterial>(MaterialEditor->GetMaterialInterface());
-				for (int32 i = 0; i < EditorMaterial->Expressions.Num(); ++i)
+				for (UMaterialExpression* Expression : EditorMaterial->GetExpressions())
 				{
-					UMaterialExpressionConstant3Vector* ColorParam = Cast<UMaterialExpressionConstant3Vector>(EditorMaterial->Expressions[i]);
+					UMaterialExpressionConstant3Vector* ColorParam = Cast<UMaterialExpressionConstant3Vector>(Expression);
 					if (ColorParam)
 					{
 						EditorMaterial->Modify();
diff --git a/Engine/Source/Developer/AssetTools/Private/AssetTypeActions/AssetTypeActions_Texture.cpp b/Engine/Source/Developer/AssetTools/Private/AssetTypeActions/AssetTypeActions_Texture.cpp
--- a/Engine/Source/Developer/AssetTools/Private/AssetTypeActions/AssetTypeActions_Texture.cpp
+++ b/Engine/Source/Developer/AssetTools/Private/AssetTypeActions/AssetTypeActions_Texture.cpp
@@ -1046,7 +1046,7 @@ void FAssetTypeActions_Texture::ConvertVTTexture(TArray<TWeakObjectPtr<UTexture>
 					MaterialTask.EnterProgressFrame();
 
 					bool FuncModified = false;
-					const TArray<UMaterialExpression*> *Expressions = Func->GetFunctionExpressions();
+					const TArray<UMaterialExpression*> *Expressions = Func->GetExpressions();
 					for (UMaterialExpression *Expr : *Expressions)
 					{
 						UMaterialExpressionTextureBase *TexExpr = Cast<UMaterialExpressionTextureBase>(Expr);
