Fixed a crash when reimporting a static mesh with a speed tree 8 file. If the static mesh wasn't imported by speed tree before it would crash the editor.

diff --git a/Engine/Plugins/Editor/SpeedTreeImporter/Source/SpeedTreeImporter/Private/SpeedTreeImportFactory.cpp b/Engine/Plugins/Editor/SpeedTreeImporter/Source/SpeedTreeImporter/Private/SpeedTreeImportFactory.cpp
--- a/Engine/Plugins/Editor/SpeedTreeImporter/Source/SpeedTreeImporter/Private/SpeedTreeImportFactory.cpp
+++ b/Engine/Plugins/Editor/SpeedTreeImporter/Source/SpeedTreeImporter/Private/SpeedTreeImportFactory.cpp
@@ -2498,7 +2498,8 @@ UObject* USpeedTreeImportFactory::FactoryCreateBinary8(UClass* InClass, UObject*
 			OldMaterials = UE::SpeedTreeImporter::Private::ClearOutOldMesh(*ExistingMesh);
 			StaticMesh = ExistingMesh;
 		}
-		else if (!bIsAutomatedImport)
+
+		if (!bIsAutomatedImport && !SpeedTreeImportData)
 		{
 			TSharedPtr<SSpeedTreeImportOptions> Options;
 
