Added underlying types to a few enums.

diff --git a/Engine/Plugins/Animation/ControlRig/Source/ControlRig/Private/Units/Debug/RigUnit_DebugHierarchy.h b/Engine/Plugins/Animation/ControlRig/Source/ControlRig/Private/Units/Debug/RigUnit_DebugHierarchy.h
--- a/Engine/Plugins/Animation/ControlRig/Source/ControlRig/Private/Units/Debug/RigUnit_DebugHierarchy.h
+++ b/Engine/Plugins/Animation/ControlRig/Source/ControlRig/Private/Units/Debug/RigUnit_DebugHierarchy.h
@@ -10,7 +10,7 @@
 UENUM()
 namespace EControlRigDrawHierarchyMode
 {
-	enum Type
+	enum Type : int
 	{
 		/** Draw as axes */
 		Axes,
diff --git a/Engine/Plugins/Editor/SpeedTreeImporter/Source/SpeedTreeImporter/Classes/SpeedTreeImportData.h b/Engine/Plugins/Editor/SpeedTreeImporter/Source/SpeedTreeImporter/Classes/SpeedTreeImportData.h
--- a/Engine/Plugins/Editor/SpeedTreeImporter/Source/SpeedTreeImporter/Classes/SpeedTreeImportData.h
+++ b/Engine/Plugins/Editor/SpeedTreeImporter/Source/SpeedTreeImporter/Classes/SpeedTreeImportData.h
@@ -25,7 +25,7 @@ class SWidget;
 
 /** Geometry import type */
 UENUM()
-enum EImportGeometryType
+enum EImportGeometryType : int
 {
 	IGT_3D UMETA(DisplayName = "3D LODs"),
 	IGT_Billboards UMETA(DisplayName = "Billboards"),
@@ -35,7 +35,7 @@ enum EImportGeometryType
 
 /** LOD type **/
 UENUM()
-enum EImportLODType
+enum EImportLODType : int
 {
 	ILT_PaintedFoliage UMETA(DisplayName = "Painted Foliage"),
 	ILT_IndividualActors UMETA(DisplayName = "Individual Actors")
diff --git a/Engine/Source/Runtime/Engine/Classes/Materials/MaterialExpressionSwitch.h b/Engine/Source/Runtime/Engine/Classes/Materials/MaterialExpressionSwitch.h
--- a/Engine/Source/Runtime/Engine/Classes/Materials/MaterialExpressionSwitch.h
+++ b/Engine/Source/Runtime/Engine/Classes/Materials/MaterialExpressionSwitch.h
@@ -13,7 +13,7 @@
 struct FPropertyChangedEvent;
 
 UENUM()
-enum ESwitchMaterialOutputType
+enum ESwitchMaterialOutputType : int
 {
 	TMMOT_Float1,
 	TMMOT_Float2,
