Removed use of the world master in the speed tree importer plugin.

diff --git a/Engine/Plugins/Editor/SpeedTreeImporter/Source/SpeedTreeImporter/Classes/SpeedTreeImportFactory.h b/Engine/Plugins/Editor/SpeedTreeImporter/Source/SpeedTreeImporter/Classes/SpeedTreeImportFactory.h
--- a/Engine/Plugins/Editor/SpeedTreeImporter/Source/SpeedTreeImporter/Classes/SpeedTreeImportFactory.h
+++ b/Engine/Plugins/Editor/SpeedTreeImporter/Source/SpeedTreeImporter/Classes/SpeedTreeImportFactory.h
@@ -38,7 +38,7 @@ public:
 #endif
 
 	TSet<UPackage*> LoadedPackages;
-	UMaterialInterface* MasterMaterial;
-	UMaterialInterface* MasterBillboardMaterial;
+	UMaterialInterface* TemplateMaterial;
+	UMaterialInterface* TemplateBillboardMaterial;
 };
 
diff --git a/Engine/Plugins/Editor/SpeedTreeImporter/Source/SpeedTreeImporter/Private/SpeedTreeImportFactory.cpp b/Engine/Plugins/Editor/SpeedTreeImporter/Source/SpeedTreeImporter/Private/SpeedTreeImportFactory.cpp
--- a/Engine/Plugins/Editor/SpeedTreeImporter/Source/SpeedTreeImporter/Private/SpeedTreeImportFactory.cpp
+++ b/Engine/Plugins/Editor/SpeedTreeImporter/Source/SpeedTreeImporter/Private/SpeedTreeImportFactory.cpp
@@ -358,10 +358,10 @@ USpeedTreeImportFactory::USpeedTreeImportFactory(const FObjectInitializer& Objec
 	Formats.Add(TEXT("st9;SpeedTree v9"));
 #endif
 
-	static ConstructorHelpers::FObjectFinder<UMaterialInterface> MasterMaterialFinder(TEXT("/SpeedTreeImporter/SpeedTree9/SpeedTreeMaster"));
-	MasterMaterial = MasterMaterialFinder.Object;
-	static ConstructorHelpers::FObjectFinder<UMaterialInterface> MasterBillboardMaterialFinder(TEXT("/SpeedTreeImporter/SpeedTree9/SpeedTreeBillboardMaster"));
-	MasterBillboardMaterial = MasterBillboardMaterialFinder.Object;
+	static ConstructorHelpers::FObjectFinder<UMaterialInterface> TemplateMaterialFinder(TEXT("/SpeedTreeImporter/SpeedTree9/SpeedTreeMaster"));
+	TemplateMaterial = TemplateMaterialFinder.Object;
+	static ConstructorHelpers::FObjectFinder<UMaterialInterface> TemplateBillboardMaterialFinder(TEXT("/SpeedTreeImporter/SpeedTree9/SpeedTreeBillboardMaster"));
+	TemplateBillboardMaterial = TemplateBillboardMaterialFinder.Object;
 }
 
 FText USpeedTreeImportFactory::GetDisplayName() const
@@ -1292,10 +1292,10 @@ UMaterialInterface* CreateSpeedTreeMaterial8(UObject* Parent, FString MaterialFu
 }
 
 
-UMaterialInterface* CreateSpeedTreeMaterial9(UObject* Parent, FString MaterialFullName, GameEngine9::CMaterial& SpeedTreeMaterial, UMaterialInterface* MasterMaterial, GameEngine9::CWindConfigSDK Wind, bool bHandleCameraFacing, int32 BillboardCount, TSet<UPackage*>& LoadedPackages, FSpeedTreeImportContext& ImportContext)
+UMaterialInterface* CreateSpeedTreeMaterial9(UObject* Parent, FString MaterialFullName, GameEngine9::CMaterial& SpeedTreeMaterial, UMaterialInterface* TemplateMaterial, GameEngine9::CWindConfigSDK Wind, bool bHandleCameraFacing, int32 BillboardCount, TSet<UPackage*>& LoadedPackages, FSpeedTreeImportContext& ImportContext)
 {
 	// Make sure we have a parent
-	if ((MasterMaterial == NULL) || !ensure(Parent))
+	if ((TemplateMaterial == NULL) || !ensure(Parent))
 	{
 		return UMaterial::GetDefaultMaterial(MD_Surface);
 	}
@@ -1362,7 +1362,7 @@ UMaterialInterface* CreateSpeedTreeMaterial9(UObject* Parent, FString MaterialFu
 	auto MaterialFactory = NewObject<UMaterialInstanceConstantFactoryNew>();
 	UMaterialInstanceConstant* UnrealMaterialInstance = (UMaterialInstanceConstant*)MaterialFactory->FactoryCreateNew(UMaterialInstanceConstant::StaticClass(), Package, *FixedMaterialName, RF_Standalone | RF_Public, NULL, GWarn);
 	FAssetRegistryModule::AssetCreated(UnrealMaterialInstance);
-	UnrealMaterialInstance->SetParentEditorOnly(MasterMaterial);
+	UnrealMaterialInstance->SetParentEditorOnly(TemplateMaterial);
 	UnrealMaterialInstance->ClearParameterValuesEditorOnly();
 	Package->SetDirtyFlag(true);
 	UnrealMaterialInstance->PreEditChange(NULL);
@@ -2997,13 +2997,13 @@ UObject* USpeedTreeImportFactory::FactoryCreateBinary9(UClass* InClass, UObject*
 				// find correct material
 				if (!MaterialMap.Contains(DrawCall.m_uiMaterialIndex))
 				{
-					// you could pick different master materials based on SpeedTree.TexturePacker()
+					// you could pick different template materials based on SpeedTree.TexturePacker()
 					bool bBillboard = (SpeedTree.BillboardInfo().LastLodIsBillboard() && (LODIndex == (SpeedTree.Lods().Count() - 1)));
-					UMaterialInterface* Master = (bBillboard ? MasterBillboardMaterial : MasterMaterial);
+					UMaterialInterface* Template = (bBillboard ? TemplateBillboardMaterial : TemplateMaterial);
 
 					GameEngine9::CMaterial SpeedTreeMaterial = SpeedTree.Materials()[DrawCall.m_uiMaterialIndex];
 					FString MaterialName = FString(SpeedTreeMaterial.Name().Data());
-					UMaterialInterface* Material = CreateSpeedTreeMaterial9(InParent, MaterialName, SpeedTreeMaterial, Master, SpeedTree.Wind(), bHasFacingData, SpeedTree.BillboardInfo().SideViewCount(), LoadedPackages, ImportContext);
+					UMaterialInterface* Material = CreateSpeedTreeMaterial9(InParent, MaterialName, SpeedTreeMaterial, Template, SpeedTree.Wind(), bHasFacingData, SpeedTree.BillboardInfo().SideViewCount(), LoadedPackages, ImportContext);
 					MaterialMap.Add(DrawCall.m_uiMaterialIndex, StaticMesh->GetStaticMaterials().Num());
 					int32 AddedMaterialIndex = StaticMesh->GetStaticMaterials().Add(FStaticMaterial(Material, FName(*MaterialName), FName(*MaterialName)));
 					const FPolygonGroupID& PolygonGroupID = MeshDescription->CreatePolygonGroup();
