[SpeedTreeImporter]
* Disable pch if header units is enabled (code depends on crt function time which is not header unit compatible in current windows sdk version)

diff --git a/Engine/Plugins/Editor/SpeedTreeImporter/Source/SpeedTreeImporter/SpeedTreeImporter.Build.cs b/Engine/Plugins/Editor/SpeedTreeImporter/Source/SpeedTreeImporter/SpeedTreeImporter.Build.cs
--- a/Engine/Plugins/Editor/SpeedTreeImporter/Source/SpeedTreeImporter/SpeedTreeImporter.Build.cs
+++ b/Engine/Plugins/Editor/SpeedTreeImporter/Source/SpeedTreeImporter/SpeedTreeImporter.Build.cs
@@ -24,6 +24,11 @@ namespace UnrealBuildTool.Rules
 				);
 				
 			AddEngineThirdPartyPrivateStaticDependencies(Target, "SpeedTree");
+
+			if (Target.bUseHeaderUnitsForPch)
+			{
+				PCHUsage = PCHUsageMode.NoPCHs; // PortableCrt.h in Speedtree uses localtime_s
+			}
 		}
 	}
 }
